support for Tsets in property bags
[RN] property Bag now supports TSets

diff --git a/Engine/Plugins/MovieScene/MovieRenderPipeline/Source/MovieRenderPipelineCore/Public/Graph/MovieGraphCommon.h b/Engine/Plugins/MovieScene/MovieRenderPipeline/Source/MovieRenderPipelineCore/Public/Graph/MovieGraphCommon.h
--- a/Engine/Plugins/MovieScene/MovieRenderPipeline/Source/MovieRenderPipelineCore/Public/Graph/MovieGraphCommon.h
+++ b/Engine/Plugins/MovieScene/MovieRenderPipeline/Source/MovieRenderPipelineCore/Public/Graph/MovieGraphCommon.h
@@ -44,6 +44,7 @@ enum class EMovieGraphContainerType : uint8
 {
 	None UMETA(Hidden),
 	Array,
+	Set,
 
 	Count UMETA(Hidden)
 };
diff --git a/Engine/Source/Developer/StructUtilsTestSuite/Private/PropertyBagTest.cpp b/Engine/Source/Developer/StructUtilsTestSuite/Private/PropertyBagTest.cpp
--- a/Engine/Source/Developer/StructUtilsTestSuite/Private/PropertyBagTest.cpp
+++ b/Engine/Source/Developer/StructUtilsTestSuite/Private/PropertyBagTest.cpp
@@ -461,7 +461,7 @@ struct FTest_Arrays : FAITestBase
 		AITEST_TRUE(TEXT("Float array should have 1 item"), FloatArray.Num() == 1);
 
 		const TValueOrError<float, EPropertyBagResult> GetDefaultFloatRes = FloatArray.GetValueFloat(FloatIndex);
-		AITEST_TRUE(TEXT("Get float should succeed immediatelly after add"), GetDefaultFloatRes.IsValid());
+		AITEST_TRUE(TEXT("Get float should succeed immediately after add"), GetDefaultFloatRes.IsValid());
 		AITEST_TRUE(TEXT("Default value for Float should be 0.0f"), FMath::IsNearlyEqual(GetDefaultFloatRes.GetValue(), 0.0f));
 
 		const EPropertyBagResult SetFloatRes = FloatArray.SetValueFloat(FloatIndex, 123.0f);
@@ -483,6 +483,141 @@ struct FTest_Arrays : FAITestBase
 };
 IMPLEMENT_AI_INSTANT_TEST(FTest_Arrays, "System.StructUtils.PropertyBag.Arrays");
 
+struct FTest_Sets : FAITestBase
+{
+	virtual bool InstantTest() override
+	{
+		static const FName FloatSetName(TEXT("FloatSet"));
+		static const FName EnumSetName(TEXT("EnumSet"));
+		static const FName StructSetName(TEXT("StructSet"));
+		static const FName ObjectSetName(TEXT("ObjectSet"));
+		static const FName ClassSetName(TEXT("ClassSet"));
+
+		FInstancedPropertyBag Bag;
+		Bag.AddProperties({
+			{ FloatSetName, EPropertyBagContainerType::Set, EPropertyBagPropertyType::Float },
+			{ EnumSetName, EPropertyBagContainerType::Set, EPropertyBagPropertyType::Enum, StaticEnum<EPropertyBagTest1>() },
+			{ StructSetName, EPropertyBagContainerType::Set, EPropertyBagPropertyType::Struct, FTestStructHashable1::StaticStruct()},
+			{ ObjectSetName, EPropertyBagContainerType::Set, EPropertyBagPropertyType::Object, UBagTestObject1::StaticClass()},
+			{ ClassSetName, EPropertyBagContainerType::Set, EPropertyBagPropertyType::Class, UBagTestObject1::StaticClass()}
+		});
+
+		//Test Numeric Type Set
+		TValueOrError<const FPropertyBagSetRef, EPropertyBagResult> FloatSetRes = Bag.GetSetRef(FloatSetName);
+		AITEST_TRUE(TEXT("Bag should have float set"), FloatSetRes.IsValid());
+
+		FPropertyBagSetRef FloatSet = FloatSetRes.GetValue();
+		float FloatValue1 = 1.f;
+		TValueOrError<bool, EPropertyBagResult> EmptySetContainsResult = FloatSet.Contains(FloatValue1);
+		AITEST_TRUE(TEXT("Float set contain result should have value"), EmptySetContainsResult.HasValue());
+		AITEST_EQUAL(TEXT("Float set contain result should be false"), EmptySetContainsResult.GetValue(), false);
+
+		const EPropertyBagResult SetFloatRes = FloatSet.AddValueFloat(FloatValue1);
+		AITEST_EQUAL(TEXT("Float set should have 1 item"), FloatSet.Num(), 1);
+		AITEST_EQUAL(TEXT("Set float should succeed"), SetFloatRes, EPropertyBagResult::Success);
+
+		TValueOrError<bool, EPropertyBagResult> FilledSetContainsResult = FloatSet.Contains(FloatValue1);
+		AITEST_TRUE(TEXT("Float set contain result should have value"), FilledSetContainsResult.HasValue());
+		AITEST_EQUAL(TEXT("Float set contain result should be true"), FilledSetContainsResult.GetValue(), true);
+
+		const float FloatValue2 = 2.f;
+		const EPropertyBagResult NewSetFloatRes = FloatSet.AddValueFloat(FloatValue2);
+		AITEST_TRUE(TEXT("Setting a new float value should succeed"), NewSetFloatRes == EPropertyBagResult::Success);
+		AITEST_EQUAL(TEXT("Float set should have 2 items"), FloatSet.Num(), 2);
+
+		const float FloatValue3 = 3.f;
+		TValueOrError<bool, EPropertyBagResult> ContainsUnknownResult = FloatSet.Contains(FloatValue3);
+		AITEST_EQUAL(TEXT("Calling Contains with a value not stored in the set should return false"), ContainsUnknownResult.GetValue(), false);
+
+		EPropertyBagResult RemoveUnknownResult = FloatSet.Remove(FloatValue3);
+		AITEST_EQUAL(TEXT("Calling Remove with a value not stored in the set should return a property not found error"), RemoveUnknownResult, EPropertyBagResult::PropertyNotFound);
+		AITEST_EQUAL(TEXT("Float set should still 2 items after failed removal"), FloatSet.Num(), 2);
+
+		const EPropertyBagResult AddExistingFloatRes = FloatSet.AddValueFloat(FloatValue2);
+		AITEST_EQUAL(TEXT("Setting an existing element to a value already present in the set should return a duplicated value error"), AddExistingFloatRes, EPropertyBagResult::DuplicatedValue);
+
+		EPropertyBagResult RemoveKnownResult = FloatSet.Remove(FloatValue2);
+		AITEST_EQUAL(TEXT("Removing an element in the set should result in success"), RemoveKnownResult, EPropertyBagResult::Success);
+		AITEST_EQUAL(TEXT("Float set should have 1 item after successfull removal"), FloatSet.Num(), 1);
+
+		const int32 IntValue = 3;
+		const EPropertyBagResult SetIntRes = FloatSet.AddValueInt32(IntValue);
+		AITEST_EQUAL(TEXT("Setting a signed integer on a float set should succeed"), SetIntRes, EPropertyBagResult::Success);
+		AITEST_TRUE(TEXT("Float set should contain the new int value"), FloatSet.Contains(IntValue).GetValue());
+
+		const uint32 UIntValue = 4;
+		const EPropertyBagResult SetUIntRes = FloatSet.AddValueUInt32(UIntValue);
+		AITEST_EQUAL(TEXT("Setting an unsigned integer on a float set should succeed"), SetUIntRes, EPropertyBagResult::Success);
+		AITEST_TRUE(TEXT("Float set should contain the new uint value"), FloatSet.Contains(UIntValue).GetValue());
+
+		const FString TestString = TEXT("TestString");
+		const EPropertyBagResult SetStringResult = FloatSet.AddValueString(TestString);
+		AITEST_NOT_EQUAL(TEXT("Setting a string on a float set should not succeed"), SetStringResult, EPropertyBagResult::Success);
+		AITEST_FALSE(TEXT("Float set should not contain the new string value"), FloatSet.Contains(TestString).GetValue());
+
+		//Test Enum Set
+		TValueOrError<const FPropertyBagSetRef, EPropertyBagResult> EnumSetRes = Bag.GetSetRef(EnumSetName);
+		AITEST_TRUE(TEXT("Bag should have Enum set property"), EnumSetRes.IsValid());
+		FPropertyBagSetRef EnumSet = EnumSetRes.GetValue();
+
+		AITEST_EQUAL(TEXT("Adding enum value to set should succeed"), EnumSet.AddValueEnum(EPropertyBagTest1::Foo), EPropertyBagResult::Success);
+		AITEST_TRUE(TEXT("We should be able to find the enum value we just added"), EnumSet.Contains(EPropertyBagTest1::Foo).GetValue());
+		AITEST_EQUAL(TEXT("Adding a different enum value to set should succeed"), EnumSet.AddValueEnum(EPropertyBagTest1::Bar), EPropertyBagResult::Success);
+		AITEST_EQUAL(TEXT("Adding an already store enum value should return a duplicated value error"), EnumSet.AddValueEnum(EPropertyBagTest1::Bar), EPropertyBagResult::DuplicatedValue);
+		AITEST_EQUAL(TEXT("Adding value from a different enum type to set should return a type mismatch error"), EnumSet.AddValueEnum(EPropertyBagTest2::Bongo), EPropertyBagResult::TypeMismatch);
+		AITEST_EQUAL(TEXT("Adding a non enum value to an enum set should return a type mismatch error"), EnumSet.AddValueInt32(1), EPropertyBagResult::TypeMismatch);
+ 
+		//Test Struct Set
+		TValueOrError<const FPropertyBagSetRef, EPropertyBagResult> StructSetRes = Bag.GetSetRef(StructSetName);
+		AITEST_TRUE(TEXT("Bag should have Struct set property"), StructSetRes.IsValid());
+		FPropertyBagSetRef StructSet = StructSetRes.GetValue();
+
+		FTestStructHashable1 TestStructInstance1;
+		TestStructInstance1.Float = 1.0f;
+
+		FTestStructHashable1 TestStructInstance2;
+		TestStructInstance2.Float = 2.0f;
+
+		FTestStructComplex ComplexStructInstance;
+
+		AITEST_EQUAL(TEXT("Adding struct value to set should succeed"), StructSet.AddValueStruct(FConstStructView::Make(TestStructInstance1)), EPropertyBagResult::Success);
+		AITEST_TRUE(TEXT("We should be able to find the struct we just added"), StructSet.Contains(FConstStructView::Make(TestStructInstance1).GetMemory()).GetValue());
+		AITEST_EQUAL(TEXT("Adding a different struct value to set should succeed"), StructSet.AddValueStruct(FConstStructView::Make(TestStructInstance2)), EPropertyBagResult::Success);
+		AITEST_EQUAL(TEXT("Adding the same struct value to set should return a duplicated value error"), StructSet.AddValueStruct(FConstStructView::Make(TestStructInstance2)), EPropertyBagResult::DuplicatedValue);
+		AITEST_EQUAL(TEXT("Adding a different struct type should return a type mismatch error"), StructSet.AddValueStruct(FConstStructView::Make(ComplexStructInstance)), EPropertyBagResult::TypeMismatch);
+		AITEST_EQUAL(TEXT("Adding a non struct type to a struct set should return a type mismatch error"), StructSet.AddValueInt32(1), EPropertyBagResult::TypeMismatch);
+
+		//Test Object Set
+		UBagTestObject1* TestObject1_Instance1 = NewObject<UBagTestObject1>();
+		UBagTestObject1* TestObject1_Instance2 = NewObject<UBagTestObject1>();
+		UBagTestObject2* TestObject2 = NewObject<UBagTestObject2>();
+		UBagTestObject1Derived* TestObject1Derived = NewObject<UBagTestObject1Derived>();
+
+		TValueOrError<const FPropertyBagSetRef, EPropertyBagResult> ObjectSetRes = Bag.GetSetRef(ObjectSetName);
+		AITEST_TRUE(TEXT("Bag should have Object set property"), ObjectSetRes.IsValid());
+		FPropertyBagSetRef ObjectSet = ObjectSetRes.GetValue();
+
+		AITEST_EQUAL(TEXT("Adding a TestObject1 type instance to set should succeed"), ObjectSet.AddValueObject(TestObject1_Instance1), EPropertyBagResult::Success);
+		AITEST_TRUE(TEXT("We should be able to find the object we just added"), ObjectSet.Contains(TestObject1_Instance1).GetValue());
+		AITEST_EQUAL(TEXT("Adding a second instance of type TestObject1 to set should succeed"), ObjectSet.AddValueObject(TestObject1_Instance2), EPropertyBagResult::Success);
+		AITEST_EQUAL(TEXT("Adding an object type derived from TestObject1 to set should succeed"), ObjectSet.AddValueObject(TestObject1Derived), EPropertyBagResult::Success);
+		AITEST_EQUAL(TEXT("Adding a TestObject2 type instance to a TestObject1 type set should return a type mismatch error"), ObjectSet.AddValueObject(TestObject2), EPropertyBagResult::TypeMismatch);
+
+		//Test Class Set
+		TValueOrError<const FPropertyBagSetRef, EPropertyBagResult> ClassSetRes = Bag.GetSetRef(ClassSetName);
+		AITEST_TRUE(TEXT("Bag should have Object set property"), ClassSetRes.IsValid());
+		FPropertyBagSetRef ClassSet = ClassSetRes.GetValue();
+
+		AITEST_EQUAL(TEXT("Adding a class to a class set should succeed"), ClassSet.AddValueClass(UBagTestObject1::StaticClass()), EPropertyBagResult::Success);
+		AITEST_TRUE(TEXT("We should be able to find the class we just added"), ClassSet.Contains(UBagTestObject1::StaticClass()).GetValue());
+		AITEST_EQUAL(TEXT("Adding a different type to a class set should return a type mismatch error"), ClassSet.AddValueClass(UBagTestObject2::StaticClass()), EPropertyBagResult::TypeMismatch);
+		AITEST_EQUAL(TEXT("Adding a derived class type to a class set should succeed"), ClassSet.AddValueClass(UBagTestObject1Derived::StaticClass()), EPropertyBagResult::Success);
+
+		return true;
+	}
+};
+IMPLEMENT_AI_INSTANT_TEST(FTest_Sets, "System.StructUtils.PropertyBag.Sets");
+
 struct FTest_SameBag : FAITestBase
 {
 	virtual bool InstantTest() override
diff --git a/Engine/Source/Developer/StructUtilsTestSuite/Private/StructUtilsTestTypes.cpp b/Engine/Source/Developer/StructUtilsTestSuite/Private/StructUtilsTestTypes.cpp
--- a/Engine/Source/Developer/StructUtilsTestSuite/Private/StructUtilsTestTypes.cpp
+++ b/Engine/Source/Developer/StructUtilsTestSuite/Private/StructUtilsTestTypes.cpp
@@ -3,3 +3,8 @@
 #include "StructUtilsTestTypes.h"
 
 #include UE_INLINE_GENERATED_CPP_BY_NAME(StructUtilsTestTypes)
+
+uint32 GetTypeHash(const FTestStructHashable1& other)
+{
+	return FCrc::MemCrc32(&other, sizeof(FTestStructHashable1));
+}
diff --git a/Engine/Source/Developer/StructUtilsTestSuite/Public/StructUtilsTestTypes.h b/Engine/Source/Developer/StructUtilsTestSuite/Public/StructUtilsTestTypes.h
--- a/Engine/Source/Developer/StructUtilsTestSuite/Public/StructUtilsTestTypes.h
+++ b/Engine/Source/Developer/StructUtilsTestSuite/Public/StructUtilsTestTypes.h
@@ -143,6 +143,19 @@ struct FTestStructSimple7 : public FTestStructSimpleBase
 	float Float = 0.0f;
 };
 
+USTRUCT()
+struct FTestStructHashable1 : public FTestStructSimpleBase
+{
+	GENERATED_BODY()
+	
+	FTestStructHashable1() = default;
+	
+	UPROPERTY()
+	float Float = 0.0f;
+};
+
+uint32 GetTypeHash(const FTestStructHashable1& other);
+
 UCLASS()
 class UBagTestObject1 : public UObject
 {
diff --git a/Engine/Source/Editor/StructUtilsEditor/Private/Customizations/PropertyBagDetails.cpp b/Engine/Source/Editor/StructUtilsEditor/Private/Customizations/PropertyBagDetails.cpp
--- a/Engine/Source/Editor/StructUtilsEditor/Private/Customizations/PropertyBagDetails.cpp
+++ b/Engine/Source/Editor/StructUtilsEditor/Private/Customizations/PropertyBagDetails.cpp
@@ -42,7 +42,7 @@ void SetPropertyDescFromPin(FPropertyBagPropertyDesc& Desc, const FEdGraphPinTyp
 		Desc.ContainerTypes.Add(EPropertyBagContainerType::Array);
 		break;
 	case EPinContainerType::Set:
-		ensureMsgf(false, TEXT("Unsuported container type [Set] "));
+		Desc.ContainerTypes.Add(EPropertyBagContainerType::Set);
 		break;
 	case EPinContainerType::Map:
 		ensureMsgf(false, TEXT("Unsuported container type [Map] "));
@@ -348,6 +348,9 @@ FEdGraphPinType GetPropertyDescAsPin(const FPropertyBagPropertyDesc& Desc)
 	case EPropertyBagContainerType::Array:
 		PinType.ContainerType = EPinContainerType::Array;
 		break;
+	case EPropertyBagContainerType::Set:
+		PinType.ContainerType = EPinContainerType::Set;
+		break;
 	default:
 		PinType.ContainerType = EPinContainerType::None;
 	}
@@ -1499,7 +1502,7 @@ TSharedPtr<SWidget> FPropertyBagDetails::MakeAddPropertyWidget(TSharedPtr<IPrope
 bool UPropertyBagSchema::SupportsPinTypeContainer(TWeakPtr<const FEdGraphSchemaAction> SchemaAction,
 		const FEdGraphPinType& PinType, const EPinContainerType& ContainerType) const
 {
-	return ContainerType == EPinContainerType::None || ContainerType == EPinContainerType::Array;
+	return ContainerType == EPinContainerType::None || ContainerType == EPinContainerType::Array || ContainerType == EPinContainerType::Set;
 }
 
 #undef LOCTEXT_NAMESPACE
diff --git a/Engine/Source/Runtime/CoreUObject/Private/StructUtils/PropertyBag.cpp b/Engine/Source/Runtime/CoreUObject/Private/StructUtils/PropertyBag.cpp
--- a/Engine/Source/Runtime/CoreUObject/Private/StructUtils/PropertyBag.cpp
+++ b/Engine/Source/Runtime/CoreUObject/Private/StructUtils/PropertyBag.cpp
@@ -254,8 +254,19 @@ namespace UE::StructUtils::Private
 						ParentName = PropName;
 						break;
 					}
+				case EPropertyBagContainerType::Set:
+					{
+						// create a set property as a container for the tail
+						const FString PropName = ParentName.IsEmpty() ? Desc.Name.ToString() : ParentName + TEXT("_InnerSet");
+						FSetProperty* SetProperty = new FSetProperty(PropertyOwner, FName(PropName), RF_Public);
+						*ValuePropertyPtr = SetProperty;
+						ValuePropertyPtr = &SetProperty->ElementProp;
+						PropertyOwner = SetProperty;
+						ParentName = PropName;
+						break;
+					}
 				default:
-					ensureMsgf(false, TEXT("Unsuported container type %s"), *UEnum::GetValueAsString(BagContainerType));
+					ensureMsgf(false, TEXT("Unsupported container type %s"), *UEnum::GetValueAsString(BagContainerType));
 					break;
 				}
 			}
@@ -2199,6 +2210,29 @@ TValueOrError<const FPropertyBagArrayRef, EPropertyBagResult> FInstancedProperty
 	return MakeValue(FPropertyBagArrayRef(*Desc, Address));
 }
 
+TValueOrError<const FPropertyBagSetRef, EPropertyBagResult> FInstancedPropertyBag::GetSetRef(const FName Name) const
+{
+	const FPropertyBagPropertyDesc* Desc = FindPropertyDescByName(Name);
+	if (Desc == nullptr)
+	{
+		return MakeError(EPropertyBagResult::PropertyNotFound);
+	}
+	check(Desc->CachedProperty);
+
+	if (Desc->ContainerTypes.GetFirstContainerType() != EPropertyBagContainerType::Set)
+	{
+		return MakeError(EPropertyBagResult::TypeMismatch);
+	}
+
+	const void* Address = GetValueAddress(Desc);
+	if (Address == nullptr)
+	{
+		return MakeError(EPropertyBagResult::PropertyNotFound);
+	}
+	
+	return MakeValue(FPropertyBagSetRef(*Desc, Address));
+}
+
 bool FInstancedPropertyBag::Identical(const FInstancedPropertyBag* Other, const uint32 PortFlags) const
 {
 	return Other && Value.Identical(&Other->Value, PortFlags);
@@ -2718,6 +2752,163 @@ EPropertyBagResult FPropertyBagArrayRef::SetValueSoftPath(const int32 Index, con
 	return UE::StructUtils::Private::SetPropertyValueAsSoftPath(&ValueDesc, GetMutableAddress(Index), InValue);
 }
 
+//----------------------------------------------------------------//
+//  FPropertyBagSetRef
+//----------------------------------------------------------------//
+
+EPropertyBagResult FPropertyBagSetRef::AddValueBool(const bool bInValue)
+{
+	return ValueDesc.IsNumericType() ? Add(bInValue) : EPropertyBagResult::TypeMismatch;
+}
+
+EPropertyBagResult FPropertyBagSetRef::AddValueByte(const uint8 InValue)
+{
+	return ValueDesc.IsNumericType() ? Add(InValue) : EPropertyBagResult::TypeMismatch;
+}
+
+EPropertyBagResult FPropertyBagSetRef::AddValueInt32(const int32 InValue)
+{
+	return ValueDesc.IsNumericType() ? Add(InValue) : EPropertyBagResult::TypeMismatch;
+}
+
+EPropertyBagResult FPropertyBagSetRef::AddValueUInt32(const uint32 InValue)
+{
+	return ValueDesc.IsNumericType() ? Add(InValue) : EPropertyBagResult::TypeMismatch;
+}
+
+EPropertyBagResult FPropertyBagSetRef::AddValueInt64(const int64 InValue)
+{
+	return ValueDesc.IsNumericType() ? Add(InValue) : EPropertyBagResult::TypeMismatch;
+}
+
+EPropertyBagResult FPropertyBagSetRef::AddValueUInt64(const uint64 InValue)
+{
+	return ValueDesc.IsNumericType() ? Add(InValue) : EPropertyBagResult::TypeMismatch;
+}
+
+EPropertyBagResult FPropertyBagSetRef::AddValueFloat(const float InValue)
+{
+	return ValueDesc.IsNumericType() ? Add(InValue) : EPropertyBagResult::TypeMismatch;
+}
+
+EPropertyBagResult FPropertyBagSetRef::AddValueDouble(const double InValue)
+{
+	return ValueDesc.IsNumericType() ? Add(InValue) : EPropertyBagResult::TypeMismatch;
+}
+
+EPropertyBagResult FPropertyBagSetRef::AddValueName(const FName InValue)
+{
+	return ValueDesc.ValueType == EPropertyBagPropertyType::Name ? Add(InValue) : EPropertyBagResult::TypeMismatch;
+}
+
+EPropertyBagResult FPropertyBagSetRef::AddValueString(const FString& InValue)
+{
+	return ValueDesc.ValueType == EPropertyBagPropertyType::String ? Add(InValue) : EPropertyBagResult::TypeMismatch;
+}
+
+EPropertyBagResult FPropertyBagSetRef::AddValueText(const FText& InValue)
+{
+	return ValueDesc.ValueType == EPropertyBagPropertyType::Text ? Add(InValue) : EPropertyBagResult::TypeMismatch;
+}
+
+EPropertyBagResult FPropertyBagSetRef::AddValueEnum(const int64 InValue, const UEnum* Enum)
+{
+	if (ValueDesc.ValueType != EPropertyBagPropertyType::Enum)
+	{
+		return EPropertyBagResult::TypeMismatch;
+	}
+
+	const FEnumProperty* EnumProperty = CastFieldChecked<FEnumProperty>(ValueDesc.CachedProperty);
+	const FNumericProperty* UnderlyingProperty = EnumProperty->GetUnderlyingProperty();
+	check(UnderlyingProperty);
+
+	if (Enum != EnumProperty->GetEnum())
+	{
+		return EPropertyBagResult::TypeMismatch;
+	}
+
+	return Add(InValue);
+}
+
+EPropertyBagResult FPropertyBagSetRef::AddValueStruct(FConstStructView InValue)
+{
+	const FStructProperty* StructProperty = CastFieldChecked<FStructProperty>(ValueDesc.CachedProperty);
+	check(StructProperty->Struct);
+
+	if (InValue.IsValid())
+	{
+		if (InValue.GetScriptStruct() != StructProperty->Struct)
+		{
+			return EPropertyBagResult::TypeMismatch;
+		}
+	}
+
+	return Add(InValue.GetMemory());
+}
+
+EPropertyBagResult FPropertyBagSetRef::AddValueObject(UObject* InValue)
+{
+	if (ValueDesc.ValueType != EPropertyBagPropertyType::Object
+		&& ValueDesc.ValueType != EPropertyBagPropertyType::SoftObject
+		&& ValueDesc.ValueType != EPropertyBagPropertyType::Class
+		&& ValueDesc.ValueType != EPropertyBagPropertyType::SoftClass)
+	{
+		return EPropertyBagResult::TypeMismatch;
+	}
+
+	const FObjectPropertyBase* ObjectProperty = CastFieldChecked<FObjectPropertyBase>(ValueDesc.CachedProperty);
+	check(ObjectProperty->PropertyClass);
+	check(ValueDesc.ValueTypeObject);
+
+	if (ValueDesc.ValueType == EPropertyBagPropertyType::Object
+		|| ValueDesc.ValueType == EPropertyBagPropertyType::SoftObject)
+	{
+		if (InValue && UE::StructUtils::Private::CanCastTo(InValue->GetClass(), ObjectProperty->PropertyClass) == false)
+		{
+			return EPropertyBagResult::TypeMismatch;
+		}
+	}
+	else
+	{
+		const UClass* Class = Cast<UClass>(InValue);
+		const UClass* PropClass = nullptr;
+
+		if (const FClassProperty* ClassProperty = CastFieldChecked<FClassProperty>(ValueDesc.CachedProperty))
+		{
+			PropClass = ClassProperty->MetaClass;
+		}
+		else if (const FSoftClassProperty* SoftClassProperty = CastFieldChecked<FSoftClassProperty>(ValueDesc.CachedProperty))
+		{
+			PropClass = SoftClassProperty->MetaClass;
+		}
+
+		if (!Class || !PropClass || !Class->IsChildOf(PropClass))
+		{
+			return EPropertyBagResult::TypeMismatch;
+		}
+	}
+
+	Add(InValue);
+
+	return EPropertyBagResult::Success;
+}
+
+EPropertyBagResult FPropertyBagSetRef::AddValueClass(UClass* InValue)
+{
+	return AddValueObject(InValue);
+}
+
+EPropertyBagResult FPropertyBagSetRef::AddValueSoftPath(const FSoftObjectPath& InValue)
+{
+	if (ValueDesc.ValueType != EPropertyBagPropertyType::SoftObject &&
+		ValueDesc.ValueType != EPropertyBagPropertyType::SoftClass)
+	{
+		return EPropertyBagResult::TypeMismatch;
+	}
+
+	return Add(InValue);
+}
+
 //----------------------------------------------------------------//
 //  UPropertyBag
 //----------------------------------------------------------------//
@@ -2939,4 +3130,4 @@ const FPropertyBagPropertyDesc* UPropertyBag::FindPropertyDescByPropertyName(con
 const FPropertyBagPropertyDesc* UPropertyBag::FindPropertyDescByProperty(const FProperty* Property) const
 {
 	return PropertyDescs.FindByPredicate([&Property](const FPropertyBagPropertyDesc& Desc) { return Desc.CachedProperty && Desc.CachedProperty == Property; });
-}
+}
\ No newline at end of file
diff --git a/Engine/Source/Runtime/CoreUObject/Public/StructUtils/PropertyBag.h b/Engine/Source/Runtime/CoreUObject/Public/StructUtils/PropertyBag.h
--- a/Engine/Source/Runtime/CoreUObject/Public/StructUtils/PropertyBag.h
+++ b/Engine/Source/Runtime/CoreUObject/Public/StructUtils/PropertyBag.h
@@ -40,6 +40,7 @@ enum class EPropertyBagContainerType : uint8
 {
 	None,
 	Array,
+	Set,
 
 	Count UMETA(Hidden)
 };
@@ -163,6 +164,7 @@ enum class EPropertyBagResult : uint8
 	TypeMismatch,		// Tried to access mismatching type (e.g. setting a struct to bool)
 	OutOfBounds,		// Tried to access an array property out of bounds.
 	PropertyNotFound,	// Could not find property of specified name.
+	DuplicatedValue,	// Tried to set an already existing Set Entry
 };
 
 USTRUCT()
@@ -338,6 +340,7 @@ struct COREUOBJECT_API FPropertyBagPropertyDesc
 
 class UPropertyBag;
 class FPropertyBagArrayRef;
+class FPropertyBagSetRef;
 
 USTRUCT()
 struct COREUOBJECT_API FInstancedPropertyBag
@@ -616,11 +619,18 @@ struct COREUOBJECT_API FInstancedPropertyBag
 
 	/**
 	 * Returns helper class to modify and access an array property.
-	 * Note: Note: The array reference is not valid after the layout of the referenced property bag has changed!
+	 * Note: The array reference is not valid after the layout of the referenced property bag has changed!
 	 * @returns helper class to modify and access arrays
 	*/
 	TValueOrError<const FPropertyBagArrayRef, EPropertyBagResult> GetArrayRef(const FName Name) const;
 
+	/**
+	 * Returns helper class to modify and access a set property.
+	 * Note: The set reference is not valid after the layout of the referenced property bag has changed!
+	 * @returns helper class to modify and access sets
+	*/
+	TValueOrError<const FPropertyBagSetRef, EPropertyBagResult> GetSetRef(const FName Name) const;
+
 	bool Identical(const FInstancedPropertyBag* Other, uint32 PortFlags) const;
 	bool Serialize(FArchive& Ar);
 	void AddStructReferencedObjects(FReferenceCollector& Collector);
@@ -838,6 +848,138 @@ public:
 	}
 };
 
+/**
+ * A reference to a set in FInstancedPropertyBag
+ * Contains helper methods to get and set properties.
+ *
+ *		FInstancedPropertyBag Bag;
+ *		Bag.AddProperties({
+ *			{ SetName, EPropertyBagContainerType::Set, EPropertyBagPropertyType::Float }
+ *		});
+ *
+ *		if (auto FloatSetRes = Bag.GetSetRef(ArrayName); FloatSetRes.IsValid())
+ *		{
+ *			FPropertyBagSetRef& FloatSet = FloatSetRes.GetValue();
+ *			FloatSet.AddValueFloat(123.f);
+ *		}
+ * 
+ * Note: The set reference is not valid after the layout of the referenced property bag has changed! 
+ */
+class COREUOBJECT_API FPropertyBagSetRef : private FScriptSetHelper
+{
+public:	
+	FORCEINLINE FPropertyBagSetRef(const FPropertyBagPropertyDesc& InDesc, const void* InSet)
+		: FScriptSetHelper(CastField<FSetProperty>(InDesc.CachedProperty), InSet)
+	{
+		const FSetProperty* SetProperty = CastField<FSetProperty>(InDesc.CachedProperty);
+		check(SetProperty);
+		check(SetProperty->ElementProp);
+		// Create dummy desc for the inner property. 
+		ValueDesc.ValueType = InDesc.ValueType;
+		ValueDesc.ValueTypeObject = InDesc.ValueTypeObject;
+		ValueDesc.CachedProperty = SetProperty->ElementProp;
+		ValueDesc.ContainerTypes = InDesc.ContainerTypes;
+		ValueDesc.ContainerTypes.PopHead();
+	}
+
+	/** Add value to set. If value is already present, it will not be added */
+	EPropertyBagResult AddValueBool(const bool bInValue);
+	EPropertyBagResult AddValueByte(const uint8 InValue);
+	EPropertyBagResult AddValueInt32(const int32 InValue);
+	EPropertyBagResult AddValueUInt32(const uint32 InValue);
+	EPropertyBagResult AddValueInt64(const int64 InValue);
+	EPropertyBagResult AddValueUInt64(const uint64 InValue);
+	EPropertyBagResult AddValueFloat(const float InValue);
+	EPropertyBagResult AddValueDouble(const double InValue);
+	EPropertyBagResult AddValueName(const FName InValue);
+	EPropertyBagResult AddValueString(const FString& InValue);
+	EPropertyBagResult AddValueText(const FText& InValue);
+	EPropertyBagResult AddValueEnum(const int64 InValue, const UEnum* Enum);
+	EPropertyBagResult AddValueStruct(FConstStructView InValue);
+	EPropertyBagResult AddValueObject(UObject* InValue);
+	EPropertyBagResult AddValueClass(UClass* InValue);
+	EPropertyBagResult AddValueSoftPath(const FSoftObjectPath& InValue);
+
+	/** Adds enum value specified type. */
+	template <typename T>
+	EPropertyBagResult AddValueEnum(const T InValue)
+	{
+		static_assert(TIsEnum<T>::Value, "Should only call this with enum types");
+		return AddValueEnum(static_cast<uint8>(InValue), StaticEnum<T>());
+	}
+
+	/** Adds struct value specified type. */
+	template <typename T>
+	EPropertyBagResult AddValueStruct(const T& InValue)
+	{
+		return AddValueStruct(FConstStructView::Make(InValue));
+	}
+
+	/** Adds object pointer value specified type. */
+	template <typename T>
+	EPropertyBagResult AddValueObject(T* InValue)
+	{
+		static_assert(TIsDerivedFrom<T, UObject>::Value, "Should only call this with object types");
+		return AddValueObject(static_cast<UObject*>(InValue));
+	}
+
+
+	/** Removes value from set if found. */
+	template <typename T>
+	EPropertyBagResult Remove(const T& Value)
+	{
+		int32 ElementIndex = FindElementIndex(&Value);
+		
+		if (ElementIndex == INDEX_NONE)
+		{
+			return EPropertyBagResult::PropertyNotFound;
+		}
+
+		RemoveAt(ElementIndex);
+		return EPropertyBagResult::Success;
+	}
+
+	/** Returns a bool specifying if the element was found or not */
+	template <typename T>
+	TValueOrError<bool, EPropertyBagResult> Contains(const T& Value) const
+	{
+        if (ValueDesc.CachedProperty == nullptr)
+        {
+            return MakeError(EPropertyBagResult::PropertyNotFound);
+        }
+
+        return MakeValue(FindElementIndex(&Value) != INDEX_NONE);
+	}
+
+	/** Returns number of elements in set. */
+	FORCEINLINE int32 Num() const
+	{
+		return FScriptSetHelper::Num();
+	}
+
+private:
+
+	FPropertyBagPropertyDesc ValueDesc;
+
+	template <typename T>
+	EPropertyBagResult Add(const T& Value)
+	{
+		if (ValueDesc.CachedProperty == nullptr)
+		{
+			return EPropertyBagResult::PropertyNotFound;
+		}
+
+		int32 ElementIndex = FindElementIndex(&Value);
+
+		if (ElementIndex != INDEX_NONE)
+		{
+			return EPropertyBagResult::DuplicatedValue;
+		}
+
+		AddElement(&Value);
+		return EPropertyBagResult::Success;
+	}
+};
 
 /**
  * Dummy types used to mark up missing types when creating property bags. These are used in the UI to display error message.
@@ -919,4 +1061,4 @@ protected:
 	std::atomic<int32> RefCount = 0;
 	
 	friend struct FInstancedPropertyBag;
-};
+};
\ No newline at end of file
