[StructUtils] fixed the TStructTypeBitSet::operator== to work as expected with emptied bitsets.

diff --git a/Engine/Source/Developer/StructUtilsTestSuite/Private/ScriptStructTypeBitSetTest.cpp b/Engine/Source/Developer/StructUtilsTestSuite/Private/ScriptStructTypeBitSetTest.cpp
--- a/Engine/Source/Developer/StructUtilsTestSuite/Private/ScriptStructTypeBitSetTest.cpp
+++ b/Engine/Source/Developer/StructUtilsTestSuite/Private/ScriptStructTypeBitSetTest.cpp
@@ -130,6 +130,14 @@ struct FStructUtilsTest_BitSetComparison : FAITestBase
 		AITEST_TRUE("", CollectionA.HasAny(CollectionSome));
 		AITEST_TRUE("HasAny is commutative", CollectionSome.HasAny(CollectionA));
 
+		FTestStructBitSet CollectionEmpty;
+		AITEST_NOT_EQUAL("Empty bitset is not equal to non-empty", CollectionEmpty, CollectionNone);
+
+		FTestStructBitSet CollectionEmptied = CollectionNone;
+		CollectionEmptied.RemoveBit(2);
+		CollectionEmptied.RemoveBit(77);
+		AITEST_EQUAL("Empty bitset is equal to an emptied one", CollectionEmpty, CollectionEmptied);
+
 		return true;
 	}
 };
diff --git a/Engine/Source/Runtime/CoreUObject/Public/StructUtils/StructTypeBitSet.h b/Engine/Source/Runtime/CoreUObject/Public/StructUtils/StructTypeBitSet.h
--- a/Engine/Source/Runtime/CoreUObject/Public/StructUtils/StructTypeBitSet.h
+++ b/Engine/Source/Runtime/CoreUObject/Public/StructUtils/StructTypeBitSet.h
@@ -73,6 +73,11 @@ struct FConstBitSetContainer : TBitArray<>
 		return !Iterator;
 	}
 
+	bool operator==(const FConstBitSetContainer& Other) const
+	{
+		return (IsEmpty() && Other.IsEmpty()) || Super::operator==(Other);
+	}
+
 	FORCEINLINE friend uint32 GetTypeHash(const FConstBitSetContainer& Instance)
 	{
 		FConstWordIterator Iterator(Instance);
