TConstStructView constructor type fix

GetScriptStruct will now return the correct type from the value provided at construction, rather than the templated type. This is consistent with the other constructors, and the equivalent signature in TStructView

#rnx

diff --git a/Engine/Plugins/Runtime/PropertyBindingUtils/Source/PropertyBindingUtilsEditor/Private/PropertyBindingExtension.cpp b/Engine/Plugins/Runtime/PropertyBindingUtils/Source/PropertyBindingUtilsEditor/Private/PropertyBindingExtension.cpp
--- a/Engine/Plugins/Runtime/PropertyBindingUtils/Source/PropertyBindingUtilsEditor/Private/PropertyBindingExtension.cpp
+++ b/Engine/Plugins/Runtime/PropertyBindingUtils/Source/PropertyBindingUtilsEditor/Private/PropertyBindingExtension.cpp
@@ -1303,7 +1303,7 @@ void FPropertyBindingExtension::ExtendWidgetRow(FDetailWidgetRow& InWidgetRow, c
 
 							if (PropertyInfoOverrides.IsEmpty())
 							{
-								InMenuBuilder.AddMenuEntry(FExecuteAction::CreateSP(InCachedBindingData, &FCachedBindingData::PromoteToParameter, PropertyName, TConstStructView(ContextStruct), TSharedPtr<const FPropertyInfoOverride>()),
+								InMenuBuilder.AddMenuEntry(FExecuteAction::CreateSP(InCachedBindingData, &FCachedBindingData::PromoteToParameter, PropertyName, TConstStructView<FPropertyBindingBindableStructDescriptor>(ContextStruct), TSharedPtr<const FPropertyInfoOverride>()),
 									MakeContextStructWidget(ContextStruct));
 							}
 							else
@@ -1320,7 +1320,7 @@ void FPropertyBindingExtension::ExtendWidgetRow(FDetailWidgetRow& InWidgetRow, c
 														InCachedBindingData,
 														&FCachedBindingData::PromoteToParameter,
 														PropertyName,
-														TConstStructView(ContextStruct),
+														TConstStructView<FPropertyBindingBindableStructDescriptor>(ContextStruct),
 														PropertyInfoOverride),
 													MakeBindingPropertyInfoWidget(PropertyInfoOverride->TypeNameText, PropertyInfoOverride->PinType));
 											}
diff --git a/Engine/Source/Developer/StructUtilsTestSuite/Private/StructViewTest.cpp b/Engine/Source/Developer/StructUtilsTestSuite/Private/StructViewTest.cpp
new file mode 100644
--- /dev/null
+++ b/Engine/Source/Developer/StructUtilsTestSuite/Private/StructViewTest.cpp
@@ -0,0 +1,49 @@
+// Copyright Epic Games, Inc. All Rights Reserved.
+
+#include "AITestsCommon.h"
+#include "StructUtils/SharedStruct.h"
+#include "StructUtilsTestTypes.h"
+
+#define LOCTEXT_NAMESPACE "StructUtilsTests"
+
+UE_DISABLE_OPTIMIZATION_SHIP
+
+namespace FInstancedStructTest
+{
+	struct FTest_StructViewCreate : FAITestBase
+	{
+		virtual bool InstantTest() override
+		{
+			{
+				FTestStructSimple TestStruct;
+
+				TStructView<FTestStructSimpleBase> StructView;
+				StructView = TStructView<FTestStructSimpleBase>(TestStruct);
+
+				AITEST_EQUAL(
+					"TStructView initialized from a value who's type is a child of the templated type should return the child type, not the templated type when calling GetScriptStruct()",
+					StructView.GetScriptStruct(),
+					FTestStructSimple::StaticStruct());
+			}
+
+			{
+				FTestStructSimple TestStruct;
+
+				TConstStructView<FTestStructSimpleBase> ConstStructView;
+				ConstStructView = TConstStructView<FTestStructSimpleBase>(TestStruct);
+
+				AITEST_EQUAL(
+					"TConstStructView initialized from a value who's type is a child of the templated type should return the child type, not the templated type when calling GetScriptStruct()",
+					ConstStructView.GetScriptStruct(),
+					FTestStructSimple::StaticStruct());
+			}
+
+			return true;
+		}
+	};
+	IMPLEMENT_AI_INSTANT_TEST(FTest_StructViewCreate, "System.StructUtils.StructView.Create");
+}
+
+UE_ENABLE_OPTIMIZATION_SHIP
+
+#undef LOCTEXT_NAMESPACE
diff --git a/Engine/Source/Runtime/CoreUObject/Public/StructUtils/StructView.h b/Engine/Source/Runtime/CoreUObject/Public/StructUtils/StructView.h
--- a/Engine/Source/Runtime/CoreUObject/Public/StructUtils/StructView.h
+++ b/Engine/Source/Runtime/CoreUObject/Public/StructUtils/StructView.h
@@ -144,7 +144,7 @@ public:
 
 	template<typename T = BaseStructT, typename = std::enable_if_t<std::is_base_of_v<BaseStructT, std::decay_t<T>>>>
 	TStructView(T& InStruct)
-		: FStructView(TBaseStructure<T>::Get(), reinterpret_cast<uint8*>(&InStruct))
+		: FStructView(T::StaticStruct(), reinterpret_cast<uint8*>(&InStruct))
 	{}
 
 	template<typename T = BaseStructT, typename = std::enable_if_t<std::is_base_of_v<BaseStructT, std::decay_t<T>>>>
@@ -345,8 +345,8 @@ struct TConstStructView : FConstStructView
 	explicit TConstStructView() = default;
 
 	template<typename T = BaseStructT, typename = std::enable_if_t<std::is_base_of_v<BaseStructT, std::decay_t<T>>>>
-	TConstStructView(const BaseStructT& InStruct)
-		: FConstStructView(TBaseStructure<T>::Get(), reinterpret_cast<const uint8*>(&InStruct))
+	TConstStructView(const T& InStruct)
+		: FConstStructView(T::StaticStruct(), reinterpret_cast<const uint8*>(&InStruct))
 	{}
 
 	template<typename T = BaseStructT, typename = std::enable_if_t<std::is_base_of_v<BaseStructT, std::decay_t<T>>>>
