PropertyBag : Add support for Maps

diff --git a/Engine/Plugins/MovieScene/MovieRenderPipeline/Source/MovieRenderPipelineCore/Public/Graph/MovieGraphCommon.h b/Engine/Plugins/MovieScene/MovieRenderPipeline/Source/MovieRenderPipelineCore/Public/Graph/MovieGraphCommon.h
--- a/Engine/Plugins/MovieScene/MovieRenderPipeline/Source/MovieRenderPipelineCore/Public/Graph/MovieGraphCommon.h
+++ b/Engine/Plugins/MovieScene/MovieRenderPipeline/Source/MovieRenderPipelineCore/Public/Graph/MovieGraphCommon.h
@@ -45,6 +45,7 @@ enum class EMovieGraphContainerType : uint8
 	None UMETA(Hidden),
 	Array,
 	Set,
+	Map,
 
 	Count UMETA(Hidden)
 };
diff --git a/Engine/Source/Developer/StructUtilsTestSuite/Private/PropertyBagTest.cpp b/Engine/Source/Developer/StructUtilsTestSuite/Private/PropertyBagTest.cpp
--- a/Engine/Source/Developer/StructUtilsTestSuite/Private/PropertyBagTest.cpp
+++ b/Engine/Source/Developer/StructUtilsTestSuite/Private/PropertyBagTest.cpp
@@ -1,7 +1,11 @@
 // Copyright Epic Games, Inc. All Rights Reserved.
 
 #include "AITestsCommon.h"
+#include "Serialization/ObjectAndNameAsStringProxyArchive.h"
+#include "Serialization/MemoryWriter.h"
+#include "Serialization/MemoryReader.h"
 #include "StructUtilsTestTypes.h"
+#include "StructUtils/PropertyBagMapRef.h"
 
 #define LOCTEXT_NAMESPACE "StructUtilsTests"
 
@@ -694,6 +698,902 @@ struct FTest_Sets : FAITestBase
 };
 IMPLEMENT_AI_INSTANT_TEST(FTest_Sets, "System.StructUtils.PropertyBag.Sets");
 
+struct FTest_Maps : FAITestBase
+{
+	virtual bool InstantTest() override
+	{
+		static const FName NameToFloatMapName(TEXT("NameToFloatMap"));
+		static const FName EnumToStructMapName(TEXT("EnumToStructMap"));
+		static const FName StringToObjectMapName(TEXT("StringToObjectMap"));
+		static const FName ClassToTextMapName(TEXT("ClassToTextMap"));
+		static const FName IntToSoftObjectMapName(TEXT("IntToSoftObjectMap"));
+		static const FName SoftClassToSoftObjectMapName(TEXT("SoftClassToSoftObjectMap"));
+		static const FName StructToSoftClassMapName(TEXT("StructToSoftClassMap"));
+		static const FName VectorToClassMapName(TEXT("VectorToClassMap"));
+
+		FInstancedPropertyBag Bag;
+		EPropertyBagAlterationResult AlterationResult = Bag.AddProperties({
+			{ NameToFloatMapName, EPropertyBagContainerType::Map, EPropertyBagPropertyType::Float, nullptr, CPF_Edit, EPropertyBagPropertyType::Name },
+			{ EnumToStructMapName, EPropertyBagContainerType::Map, EPropertyBagPropertyType::Struct, FTestStructHashable1::StaticStruct(), CPF_Edit, EPropertyBagPropertyType::Enum, StaticEnum<EPropertyBagTest1>() },
+			{ StringToObjectMapName, EPropertyBagContainerType::Map, EPropertyBagPropertyType::Object, UBagTestObject1::StaticClass(), CPF_Edit, EPropertyBagPropertyType::String },
+			{ ClassToTextMapName, EPropertyBagContainerType::Map, EPropertyBagPropertyType::Text, nullptr, CPF_Edit, EPropertyBagPropertyType::Class, UBagTestObject1::StaticClass() },
+			{ IntToSoftObjectMapName, EPropertyBagContainerType::Map, EPropertyBagPropertyType::SoftObject, UBagTestObject1::StaticClass(), CPF_Edit, EPropertyBagPropertyType::Int32 },
+			{ SoftClassToSoftObjectMapName, EPropertyBagContainerType::Map, EPropertyBagPropertyType::SoftObject, UBagTestObject1::StaticClass(), CPF_Edit, EPropertyBagPropertyType::SoftClass, UBagTestObject2::StaticClass() },
+			{ StructToSoftClassMapName, EPropertyBagContainerType::Map, EPropertyBagPropertyType::SoftClass, UBagTestObject1::StaticClass(), CPF_Edit, EPropertyBagPropertyType::Struct, FTestStructHashable1::StaticStruct() },
+		});
+		AITEST_EQUAL(TEXT("Bag.AddProperties should succeed"), AlterationResult, EPropertyBagAlterationResult::Success);
+		
+		{
+			AlterationResult = Bag.AddMapProperty(VectorToClassMapName, EPropertyBagPropertyType::Struct, EPropertyBagPropertyType::Class, TBaseStructure<FVector>::Get(), UBagTestObject1::StaticClass());
+			AITEST_EQUAL(TEXT("Bag.AddMapProperty should succeed"), AlterationResult, EPropertyBagAlterationResult::Success);
+		}
+
+		{
+			AlterationResult = Bag.AddMapProperty(FName("BoolKeyTypeIsNotHashable"), EPropertyBagPropertyType::Bool, EPropertyBagPropertyType::Struct, nullptr, FTestStructHashable1::StaticStruct());
+			AITEST_EQUAL(TEXT("Bag.AddMapProperty should fail with PropertyKeyTypeNotHashable for Bool KeyType"), AlterationResult, EPropertyBagAlterationResult::PropertyKeyTypeNotHashable);
+
+			AlterationResult = Bag.AddMapProperty(FName("TextKeyTypeIsNotHashable"), EPropertyBagPropertyType::Text, EPropertyBagPropertyType::Struct, nullptr, FTestStructHashable1::StaticStruct());
+			AITEST_EQUAL(TEXT("Bag.AddMapProperty should fail with PropertyKeyTypeNotHashable for Text KeyType"), AlterationResult, EPropertyBagAlterationResult::PropertyKeyTypeNotHashable);
+
+			AlterationResult = Bag.AddMapProperty(FName("RotatorKeyTypeIsNotHashable"), EPropertyBagPropertyType::Struct, EPropertyBagPropertyType::Struct, TBaseStructure<FRotator>::Get(), FTestStructHashable1::StaticStruct());
+			AITEST_EQUAL(TEXT("Bag.AddMapProperty should fail with PropertyKeyTypeNotHashable for FRotator KeyType"), AlterationResult, EPropertyBagAlterationResult::PropertyKeyTypeNotHashable);
+		}
+
+		// map: Name->Float
+		{
+			TValueOrError<const FPropertyBagMapRef, EPropertyBagResult> GetMapResult = Bag.GetMapRef(NameToFloatMapName);
+			AITEST_TRUE(TEXT("Bag should have map: Name->Float"), GetMapResult.IsValid());
+
+			FPropertyBagMapRef NameToFloatMap = GetMapResult.GetValue();
+			AITEST_EQUAL(TEXT("NameToFloatMap should be empty"), NameToFloatMap.Num(), 0);
+
+			TValueOrError<bool, EPropertyBagResult> ContainsResult = NameToFloatMap.Contains(FName("NotGonnaFindIt"));
+			AITEST_TRUE(TEXT("NameToFloatMap.Contains result should have value"), ContainsResult.HasValue());
+			AITEST_EQUAL(TEXT("NameToFloatMap.Contains result should be false"), ContainsResult.GetValue(), false);
+
+			ContainsResult = NameToFloatMap.Contains(42);
+			AITEST_TRUE(TEXT("NameToFloatMap.Contains result should have error"), ContainsResult.HasError());
+			AITEST_EQUAL(TEXT("NameToFloatMap.Contains result should be TypeMismatch"), ContainsResult.GetError(), EPropertyBagResult::TypeMismatch);
+
+			EPropertyBagResult AddResult = NameToFloatMap.Add(FName("PI"), 3.14f);
+			AITEST_EQUAL(TEXT("NameToFloatMap.Add should succeed"), AddResult, EPropertyBagResult::Success);
+			AITEST_EQUAL(TEXT("NameToFloatMap should have 1 pair"), NameToFloatMap.Num(), 1);
+			
+			AddResult = NameToFloatMap.Add(FName("Wrong"), FName("Type"));
+			AITEST_EQUAL(TEXT("NameToFloatMap.Add should fail with TypeMismatch"), AddResult, EPropertyBagResult::TypeMismatch);
+			AITEST_EQUAL(TEXT("NameToFloatMap should have 1 pair"), NameToFloatMap.Num(), 1);
+
+			AddResult = NameToFloatMap.Add(FName("PI"), 3.14159f);
+			AITEST_EQUAL(TEXT("NameToFloatMap.Add should succeed"), AddResult, EPropertyBagResult::Success);
+			AITEST_EQUAL(TEXT("NameToFloatMap should have 1 pair"), NameToFloatMap.Num(), 1);
+
+			AddResult = NameToFloatMap.Add(FName("Answer"), 42.f);
+			AITEST_EQUAL(TEXT("NameToFloatMap.Add should succeed"), AddResult, EPropertyBagResult::Success);
+			AITEST_EQUAL(TEXT("NameToFloatMap should have 2 pairs"), NameToFloatMap.Num(), 2);
+
+			ContainsResult = NameToFloatMap.Contains(FName("Answer"));
+			AITEST_TRUE(TEXT("NameToFloatMap.Contains result should have value"), ContainsResult.HasValue());
+			AITEST_EQUAL(TEXT("NameToFloatMap.Contains result should be true"), ContainsResult.GetValue(), true);
+
+			TValueOrError<float, EPropertyBagResult> FindResult = NameToFloatMap.Find<FName, float>(FName("PI"));
+			AITEST_TRUE(TEXT("NameToFloatMap.Find result should have value"), FindResult.HasValue());
+			AITEST_EQUAL(TEXT("NameToFloatMap.Find result should be expected value"), FindResult.GetValue(), 3.14159f);
+
+			FindResult = NameToFloatMap.Find<FName, float>(FName("NotGonnaFindIt"));
+			AITEST_TRUE(TEXT("NameToFloatMap.Find result should have error"), FindResult.HasError());
+			AITEST_EQUAL(TEXT("NameToFloatMap.Find result should be PropertyNotFound"), FindResult.GetError(), EPropertyBagResult::PropertyNotFound);
+
+			EPropertyBagResult RemoveResult = NameToFloatMap.Remove(FName("Answer"));
+			AITEST_EQUAL(TEXT("NameToFloatMap.Remove should succeed"), RemoveResult, EPropertyBagResult::Success);
+			AITEST_EQUAL(TEXT("NameToFloatMap should have 1 pair"), NameToFloatMap.Num(), 1);
+
+			ContainsResult = NameToFloatMap.Contains(FName("Answer"));
+			AITEST_TRUE(TEXT("NameToFloatMap.Contains result should have value"), ContainsResult.HasValue());
+			AITEST_EQUAL(TEXT("NameToFloatMap.Contains result should be false"), ContainsResult.GetValue(), false);
+
+			NameToFloatMap.Empty();
+			AITEST_EQUAL(TEXT("NameToFloatMap should be empty"), NameToFloatMap.Num(), 0);
+
+			ContainsResult = NameToFloatMap.Contains(FName("PI"));
+			AITEST_TRUE(TEXT("NameToFloatMap.Contains result should have value"), ContainsResult.HasValue());
+			AITEST_EQUAL(TEXT("NameToFloatMap.Contains result should be false"), ContainsResult.GetValue(), false);
+		}
+
+		// map: Enum->Struct
+		{
+			TValueOrError<const FPropertyBagMapRef, EPropertyBagResult> GetMapResult = Bag.GetMapRef(EnumToStructMapName);
+			AITEST_TRUE(TEXT("Bag should have map: Enum->Struct"), GetMapResult.IsValid());
+
+			FPropertyBagMapRef EnumToStructMap = GetMapResult.GetValue();
+			AITEST_EQUAL(TEXT("EnumToStructMap should be empty"), EnumToStructMap.Num(), 0);
+
+			TValueOrError<bool, EPropertyBagResult> ContainsResult = EnumToStructMap.Contains(EPropertyBagTest1::Foo);
+			AITEST_TRUE(TEXT("EnumToStructMap.Contains result should have value"), ContainsResult.HasValue());
+			AITEST_EQUAL(TEXT("EnumToStructMap.Contains result should be false"), ContainsResult.GetValue(), false);
+
+			ContainsResult = EnumToStructMap.Contains(EPropertyBagTest2::Bingo);
+			AITEST_TRUE(TEXT("EnumToStructMap.Contains result should have error"), ContainsResult.HasError());
+			AITEST_EQUAL(TEXT("EnumToStructMap.Contains result should be TypeMismatch"), ContainsResult.GetError(), EPropertyBagResult::TypeMismatch);
+
+			FTestStructHashable1 Struct1;
+			Struct1.Float = 3.14f;
+
+			EPropertyBagResult AddResult = EnumToStructMap.Add(EPropertyBagTest1::Foo, Struct1);
+			AITEST_EQUAL(TEXT("EnumToStructMap.Add should succeed"), AddResult, EPropertyBagResult::Success);
+			AITEST_EQUAL(TEXT("EnumToStructMap should have 1 pair"), EnumToStructMap.Num(), 1);
+
+			AddResult = EnumToStructMap.Add(EPropertyBagTest1::Bar, FName("WrongType"));
+			AITEST_EQUAL(TEXT("EnumToStructMap.Add should fail with TypeMismatch"), AddResult, EPropertyBagResult::TypeMismatch);
+			AITEST_EQUAL(TEXT("EnumToStructMap should have 1 pair"), EnumToStructMap.Num(), 1);
+
+			FTestStructHashable1 Struct2;
+			Struct2.Float = 3.14159f;
+
+			AddResult = EnumToStructMap.Add(EPropertyBagTest1::Foo, Struct2);
+			AITEST_EQUAL(TEXT("EnumToStructMap.Add should succeed"), AddResult, EPropertyBagResult::Success);
+			AITEST_EQUAL(TEXT("EnumToStructMap should have 1 pair"), EnumToStructMap.Num(), 1);
+
+			FTestStructHashable1 Struct3;
+			Struct3.Float = 42.f;
+
+			AddResult = EnumToStructMap.Add(EPropertyBagTest1::Bar, Struct3);
+			AITEST_EQUAL(TEXT("EnumToStructMap.Add should succeed"), AddResult, EPropertyBagResult::Success);
+			AITEST_EQUAL(TEXT("EnumToStructMap should have 2 pairs"), EnumToStructMap.Num(), 2);
+
+			ContainsResult = EnumToStructMap.Contains(EPropertyBagTest1::Bar);
+			AITEST_TRUE(TEXT("EnumToStructMap.Contains result should have value"), ContainsResult.HasValue());
+			AITEST_EQUAL(TEXT("EnumToStructMap.Contains result should be true"), ContainsResult.GetValue(), true);
+
+			TValueOrError<FTestStructHashable1, EPropertyBagResult> FindResult = EnumToStructMap.Find<EPropertyBagTest1, FTestStructHashable1>(EPropertyBagTest1::Bar);
+			AITEST_TRUE(TEXT("EnumToStructMap.Find result should have value"), FindResult.HasValue());
+			AITEST_EQUAL(TEXT("EnumToStructMap.Find result should be expected value"), FindResult.GetValue().Float, 42.f);
+
+			FindResult = EnumToStructMap.Find<EPropertyBagTest1, FTestStructHashable1>(EPropertyBagTest1::Baz);
+			AITEST_TRUE(TEXT("EnumToStructMap.Find result should have error"), FindResult.HasError());
+			AITEST_EQUAL(TEXT("EnumToStructMap.Find result should be PropertyNotFound"), FindResult.GetError(), EPropertyBagResult::PropertyNotFound);
+
+			EPropertyBagResult RemoveResult = EnumToStructMap.Remove(EPropertyBagTest1::Bar);
+			AITEST_EQUAL(TEXT("EnumToStructMap.Remove should succeed"), RemoveResult, EPropertyBagResult::Success);
+			AITEST_EQUAL(TEXT("EnumToStructMap should have 1 pair"), EnumToStructMap.Num(), 1);
+
+			ContainsResult = EnumToStructMap.Contains(EPropertyBagTest1::Bar);
+			AITEST_TRUE(TEXT("EnumToStructMap.Contains result should have value"), ContainsResult.HasValue());
+			AITEST_EQUAL(TEXT("EnumToStructMap.Contains result should be false"), ContainsResult.GetValue(), false);
+
+			EnumToStructMap.Empty();
+			AITEST_EQUAL(TEXT("EnumToStructMap should be empty"), EnumToStructMap.Num(), 0);
+
+			ContainsResult = EnumToStructMap.Contains(EPropertyBagTest1::Foo);
+			AITEST_TRUE(TEXT("EnumToStructMap.Contains result should have value"), ContainsResult.HasValue());
+			AITEST_EQUAL(TEXT("EnumToStructMap.Contains result should be false"), ContainsResult.GetValue(), false);
+		}
+
+		// map: String->Object
+		{
+			TValueOrError<const FPropertyBagMapRef, EPropertyBagResult> GetMapResult = Bag.GetMapRef(StringToObjectMapName);
+			AITEST_TRUE(TEXT("Bag should have map: String->Object"), GetMapResult.IsValid());
+
+			FPropertyBagMapRef StringToObjectMap = GetMapResult.GetValue();
+			AITEST_EQUAL(TEXT("StringToObjectMap should be empty"), StringToObjectMap.Num(), 0);
+
+			TValueOrError<bool, EPropertyBagResult> ContainsResult = StringToObjectMap.Contains(FString(TEXT("Key")));
+			AITEST_TRUE(TEXT("StringToObjectMap.Contains result should have value"), ContainsResult.HasValue());
+			AITEST_EQUAL(TEXT("StringToObjectMap.Contains result should be false"), ContainsResult.GetValue(), false);
+
+			UBagTestObject1* TestObject1_Instance1 = NewObject<UBagTestObject1>();
+
+			EPropertyBagResult AddResult = StringToObjectMap.Add(FString(TEXT("Key")), TestObject1_Instance1);
+			AITEST_EQUAL(TEXT("StringToObjectMap.Add should succeed"), AddResult, EPropertyBagResult::Success);
+			AITEST_EQUAL(TEXT("StringToObjectMap should have 1 pair"), StringToObjectMap.Num(), 1);
+
+			UBagTestObject2* TestObject2_Instance1 = NewObject<UBagTestObject2>();
+
+			AddResult = StringToObjectMap.Add(FString(TEXT("AnotherKey")), TestObject2_Instance1);
+			AITEST_EQUAL(TEXT("StringToObjectMap.Add should fail with TypeMismatch"), AddResult, EPropertyBagResult::TypeMismatch);
+			AITEST_EQUAL(TEXT("StringToObjectMap should have 1 pair"), StringToObjectMap.Num(), 1);
+
+			UBagTestObject1* TestObject1_Instance2 = NewObject<UBagTestObject1>();
+
+			AddResult = StringToObjectMap.Add(FString(TEXT("Key")), TestObject1_Instance2);
+			AITEST_EQUAL(TEXT("StringToObjectMap.Add should succeed"), AddResult, EPropertyBagResult::Success);
+			AITEST_EQUAL(TEXT("StringToObjectMap should have 1 pair"), StringToObjectMap.Num(), 1);
+
+			UBagTestObject1* TestObject1_Instance3 = NewObject<UBagTestObject1>();
+
+			AddResult = StringToObjectMap.Add(FString(TEXT("AnotherKey")), TestObject1_Instance3);
+			AITEST_EQUAL(TEXT("StringToObjectMap.Add should succeed"), AddResult, EPropertyBagResult::Success);
+			AITEST_EQUAL(TEXT("StringToObjectMap should have 2 pairs"), StringToObjectMap.Num(), 2);
+
+			ContainsResult = StringToObjectMap.Contains(FString(TEXT("AnotherKey")));
+			AITEST_TRUE(TEXT("StringToObjectMap.Contains result should have value"), ContainsResult.HasValue());
+			AITEST_EQUAL(TEXT("StringToObjectMap.Contains result should be true"), ContainsResult.GetValue(), true);
+
+			TValueOrError<UBagTestObject1*, EPropertyBagResult> FindResult = StringToObjectMap.Find<FString, UBagTestObject1*>(FString(TEXT("AnotherKey")));
+			AITEST_TRUE(TEXT("StringToObjectMap.Find result should have value"), FindResult.HasValue());
+			AITEST_EQUAL(TEXT("StringToObjectMap.Find result should be expected value"), FindResult.GetValue(), TestObject1_Instance3);
+
+			FindResult = StringToObjectMap.Find<FString, UBagTestObject1*>(FString(TEXT("MissingKey")));
+			AITEST_TRUE(TEXT("StringToObjectMap.Find result should have error"), FindResult.HasError());
+			AITEST_EQUAL(TEXT("StringToObjectMap.Find result should be PropertyNotFound"), FindResult.GetError(), EPropertyBagResult::PropertyNotFound);
+
+			EPropertyBagResult RemoveResult = StringToObjectMap.Remove(FString(TEXT("AnotherKey")));
+			AITEST_EQUAL(TEXT("StringToObjectMap.Remove should succeed"), RemoveResult, EPropertyBagResult::Success);
+			AITEST_EQUAL(TEXT("StringToObjectMap should have 1 pair"), StringToObjectMap.Num(), 1);
+
+			ContainsResult = StringToObjectMap.Contains(FString(TEXT("AnotherKey")));
+			AITEST_TRUE(TEXT("StringToObjectMap.Contains result should have value"), ContainsResult.HasValue());
+			AITEST_EQUAL(TEXT("StringToObjectMap.Contains result should be false"), ContainsResult.GetValue(), false);
+
+			StringToObjectMap.Empty();
+			AITEST_EQUAL(TEXT("StringToObjectMap should be empty"), StringToObjectMap.Num(), 0);
+
+			ContainsResult = StringToObjectMap.Contains(FString(TEXT("Key")));
+			AITEST_TRUE(TEXT("StringToObjectMap.Contains result should have value"), ContainsResult.HasValue());
+			AITEST_EQUAL(TEXT("StringToObjectMap.Contains result should be false"), ContainsResult.GetValue(), false);
+		}
+
+		// map: Class->Text
+		{
+			TValueOrError<const FPropertyBagMapRef, EPropertyBagResult> GetMapResult = Bag.GetMapRef(ClassToTextMapName);
+			AITEST_TRUE(TEXT("Bag should have map: Class->Text"), GetMapResult.IsValid());
+
+			FPropertyBagMapRef ClassToTextMap = GetMapResult.GetValue();
+			AITEST_EQUAL(TEXT("ClassToTextMap should be empty"), ClassToTextMap.Num(), 0);
+
+			TValueOrError<bool, EPropertyBagResult> ContainsResult = ClassToTextMap.Contains(UBagTestObject1::StaticClass());
+			AITEST_TRUE(TEXT("ClassToTextMap.Contains result should have value"), ContainsResult.HasValue());
+			AITEST_EQUAL(TEXT("ClassToTextMap.Contains result should be false"), ContainsResult.GetValue(), false);
+
+			EPropertyBagResult AddResult = ClassToTextMap.Add(UBagTestObject1::StaticClass(), FText::FromString(TEXT("Value")));
+			AITEST_EQUAL(TEXT("ClassToTextMap.Add should succeed"), AddResult, EPropertyBagResult::Success);
+			AITEST_EQUAL(TEXT("ClassToTextMap should have 1 pair"), ClassToTextMap.Num(), 1);
+
+			AddResult = ClassToTextMap.Add(UBagTestObject2::StaticClass(), FText::FromString(TEXT("WrongKeyType")));
+			AITEST_EQUAL(TEXT("ClassToTextMap.Add should fail with TypeMismatch"), AddResult, EPropertyBagResult::TypeMismatch);
+			AITEST_EQUAL(TEXT("ClassToTextMap should have 1 pair"), ClassToTextMap.Num(), 1);
+
+			AddResult = ClassToTextMap.Add(UBagTestObject1::StaticClass(), FText::FromString(TEXT("AnotherValue")));
+			AITEST_EQUAL(TEXT("ClassToTextMap.Add should succeed"), AddResult, EPropertyBagResult::Success);
+			AITEST_EQUAL(TEXT("ClassToTextMap should have 1 pair"), ClassToTextMap.Num(), 1);
+
+			AddResult = ClassToTextMap.Add(UBagTestObject1Derived::StaticClass(), FText::FromString(TEXT("ValueForSecondPair")));
+			AITEST_EQUAL(TEXT("ClassToTextMap.Add should succeed"), AddResult, EPropertyBagResult::Success);
+			AITEST_EQUAL(TEXT("ClassToTextMap should have 2 pairs"), ClassToTextMap.Num(), 2);
+
+			ContainsResult = ClassToTextMap.Contains(UBagTestObject1Derived::StaticClass());
+			AITEST_TRUE(TEXT("ClassToTextMap.Contains result should have value"), ContainsResult.HasValue());
+			AITEST_EQUAL(TEXT("ClassToTextMap.Contains result should be true"), ContainsResult.GetValue(), true);
+
+			TValueOrError<FText, EPropertyBagResult> FindResult = ClassToTextMap.Find<UClass*, FText>(UBagTestObject1Derived::StaticClass());
+			AITEST_TRUE(TEXT("ClassToTextMap.Find result should have value"), FindResult.HasValue());
+			AITEST_EQUAL(TEXT("ClassToTextMap.Find result should be expected value"), FindResult.GetValue(), FText::FromString(TEXT("ValueForSecondPair")));
+
+			FindResult = ClassToTextMap.Find<UClass*, FText>(UBagTestObject1Derived2::StaticClass());
+			AITEST_TRUE(TEXT("ClassToTextMap.Find result should have error"), FindResult.HasError());
+			AITEST_EQUAL(TEXT("ClassToTextMap.Find result should be PropertyNotFound"), FindResult.GetError(), EPropertyBagResult::PropertyNotFound);
+
+			EPropertyBagResult RemoveResult = ClassToTextMap.Remove(UBagTestObject1Derived::StaticClass());
+			AITEST_EQUAL(TEXT("ClassToTextMap.Remove should succeed"), RemoveResult, EPropertyBagResult::Success);
+			AITEST_EQUAL(TEXT("ClassToTextMap should have 1 pair"), ClassToTextMap.Num(), 1);
+
+			ContainsResult = ClassToTextMap.Contains(UBagTestObject1Derived::StaticClass());
+			AITEST_TRUE(TEXT("ClassToTextMap.Contains result should have value"), ContainsResult.HasValue());
+			AITEST_EQUAL(TEXT("ClassToTextMap.Contains result should be false"), ContainsResult.GetValue(), false);
+
+			ClassToTextMap.Empty();
+			AITEST_EQUAL(TEXT("ClassToTextMap should be empty"), ClassToTextMap.Num(), 0);
+
+			ContainsResult = ClassToTextMap.Contains(UBagTestObject1::StaticClass());
+			AITEST_TRUE(TEXT("ClassToTextMap.Contains result should have value"), ContainsResult.HasValue());
+			AITEST_EQUAL(TEXT("ClassToTextMap.Contains result should be false"), ContainsResult.GetValue(), false);
+		}
+
+		// map: int->SoftObject
+		{
+			TValueOrError<const FPropertyBagMapRef, EPropertyBagResult> GetMapResult = Bag.GetMapRef(IntToSoftObjectMapName);
+			AITEST_TRUE(TEXT("Bag should have map: int->SoftObject"), GetMapResult.IsValid());
+
+			FPropertyBagMapRef IntToSoftObjectMap = GetMapResult.GetValue();
+			AITEST_EQUAL(TEXT("IntToSoftObjectMap should be empty"), IntToSoftObjectMap.Num(), 0);
+
+			TValueOrError<bool, EPropertyBagResult> ContainsResult = IntToSoftObjectMap.Contains(42);
+			AITEST_TRUE(TEXT("IntToSoftObjectMap.Contains result should have value"), ContainsResult.HasValue());
+			AITEST_EQUAL(TEXT("IntToSoftObjectMap.Contains result should be false"), ContainsResult.GetValue(), false);
+
+			EPropertyBagResult AddResult = IntToSoftObjectMap.Add(42, FSoftObjectPath("/Game/Answer"));
+			AITEST_EQUAL(TEXT("IntToSoftObjectMap.Add<FSoftObjectPath> should succeed"), AddResult, EPropertyBagResult::Success);
+			AITEST_EQUAL(TEXT("IntToSoftObjectMap should have 1 pair"), IntToSoftObjectMap.Num(), 1);
+
+			UBagTestObject2* TestObject2_Instance1 = NewObject<UBagTestObject2>();
+
+			AddResult = IntToSoftObjectMap.Add(13, TestObject2_Instance1);
+			AITEST_EQUAL(TEXT("IntToSoftObjectMap.Add<UBagTestObject2> should fail with TypeMismatch"), AddResult, EPropertyBagResult::TypeMismatch);
+			AITEST_EQUAL(TEXT("IntToSoftObjectMap should have 1 pair"), IntToSoftObjectMap.Num(), 1);
+			
+			UBagTestObject1* TestObject1_Instance1 = NewObject<UBagTestObject1>();
+
+			AddResult = IntToSoftObjectMap.Add(128, TestObject1_Instance1);
+			AITEST_EQUAL(TEXT("IntToSoftObjectMap.Add<UBagTestObject1> should succeed"), AddResult, EPropertyBagResult::Success);
+			AITEST_EQUAL(TEXT("IntToSoftObjectMap should have 2 pairs"), IntToSoftObjectMap.Num(), 2);
+
+			UBagTestObject1* TestObject1_Instance2 = NewObject<UBagTestObject1>();
+
+			AddResult = IntToSoftObjectMap.Add(13, FSoftObjectPath(TestObject1_Instance2));
+			AITEST_EQUAL(TEXT("IntToSoftObjectMap.Add<FSoftObjectPath> should succeed"), AddResult, EPropertyBagResult::Success);
+			AITEST_EQUAL(TEXT("IntToSoftObjectMap should have 3 pairs"), IntToSoftObjectMap.Num(), 3);
+						
+			ContainsResult = IntToSoftObjectMap.Contains(42);
+			AITEST_TRUE(TEXT("IntToSoftObjectMap.Contains result should have value"), ContainsResult.HasValue());
+			AITEST_EQUAL(TEXT("IntToSoftObjectMap.Contains result should be true"), ContainsResult.GetValue(), true);
+
+			TValueOrError<FSoftObjectPath, EPropertyBagResult> FindAsSoftObjectPathResult = IntToSoftObjectMap.Find<int, FSoftObjectPath>(13);
+			AITEST_TRUE(TEXT("IntToSoftObjectMap.Find<FSoftObjectPath> result should have value"), FindAsSoftObjectPathResult.HasValue());
+			AITEST_EQUAL(TEXT("IntToSoftObjectMap.Find<FSoftObjectPath> result should be expected value"), FindAsSoftObjectPathResult.GetValue(), FSoftObjectPath(TestObject1_Instance2));
+
+			TValueOrError<FSoftObjectPtr, EPropertyBagResult> FindAsSoftObjectPtrResult = IntToSoftObjectMap.Find<int, FSoftObjectPtr>(13);
+			AITEST_TRUE(TEXT("IntToSoftObjectMap.Find<FSoftObjectPtr> result should have value"), FindAsSoftObjectPtrResult.HasValue());
+			AITEST_EQUAL(TEXT("IntToSoftObjectMap.Find<FSoftObjectPtr> result should be expected value"), FindAsSoftObjectPtrResult.GetValue(), FSoftObjectPtr(TestObject1_Instance2));
+
+			TValueOrError<UBagTestObject1*, EPropertyBagResult> FindAsUObjectPtrResult = IntToSoftObjectMap.Find<int, UBagTestObject1*>(13);
+			AITEST_TRUE(TEXT("IntToSoftObjectMap.Find<UBagTestObject1> result should have value"), FindAsUObjectPtrResult.HasValue());
+			AITEST_EQUAL(TEXT("IntToSoftObjectMap.Find<UBagTestObject1> result should be expected value"), FindAsUObjectPtrResult.GetValue(), TestObject1_Instance2);
+
+			FindAsSoftObjectPathResult = IntToSoftObjectMap.Find<int, FSoftObjectPath>(42);
+			AITEST_TRUE(TEXT("IntToSoftObjectMap.Find<FSoftObjectPath> result should have value"), FindAsSoftObjectPathResult.HasValue());
+			AITEST_EQUAL(TEXT("IntToSoftObjectMap.Find<FSoftObjectPath> result should be expected value"), FindAsSoftObjectPathResult.GetValue(), FSoftObjectPath("/Game/Answer"));
+
+			FindAsSoftObjectPtrResult = IntToSoftObjectMap.Find<int, FSoftObjectPtr>(42);
+			AITEST_TRUE(TEXT("IntToSoftObjectMap.Find<FSoftObjectPtr> result should have value"), FindAsSoftObjectPtrResult.HasValue());
+			AITEST_EQUAL(TEXT("IntToSoftObjectMap.Find<FSoftObjectPtr> result should be expected value"), FindAsSoftObjectPtrResult.GetValue(), FSoftObjectPtr(FSoftObjectPath("/Game/Answer")));
+
+			FindAsUObjectPtrResult = IntToSoftObjectMap.Find<int, UBagTestObject1*>(42);
+			AITEST_TRUE(TEXT("IntToSoftObjectMap.Find<UBagTestObject1> result should have value"), FindAsUObjectPtrResult.HasValue());
+			AITEST_NULL(TEXT("IntToSoftObjectMap.Find<UBagTestObject1> result should be null"), FindAsUObjectPtrResult.GetValue());
+
+			EPropertyBagResult RemoveResult = IntToSoftObjectMap.Remove(128);
+			AITEST_EQUAL(TEXT("IntToSoftObjectMap.Remove should succeed"), RemoveResult, EPropertyBagResult::Success);
+			AITEST_EQUAL(TEXT("IntToSoftObjectMap should have 2 pair"), IntToSoftObjectMap.Num(), 2);
+
+			ContainsResult = IntToSoftObjectMap.Contains(128);
+			AITEST_TRUE(TEXT("IntToSoftObjectMap.Contains result should have value"), ContainsResult.HasValue());
+			AITEST_EQUAL(TEXT("IntToSoftObjectMap.Contains result should be false"), ContainsResult.GetValue(), false);
+
+			IntToSoftObjectMap.Empty();
+			AITEST_EQUAL(TEXT("IntToSoftObjectMap should be empty"), IntToSoftObjectMap.Num(), 0);
+
+			ContainsResult = IntToSoftObjectMap.Contains(42);
+			AITEST_TRUE(TEXT("IntToSoftObjectMap.Contains result should have value"), ContainsResult.HasValue());
+			AITEST_EQUAL(TEXT("IntToSoftObjectMap.Contains result should be false"), ContainsResult.GetValue(), false);
+		}
+
+		// map: SoftClass->SoftObject
+		{
+			TValueOrError<const FPropertyBagMapRef, EPropertyBagResult> GetMapResult = Bag.GetMapRef(SoftClassToSoftObjectMapName);
+			AITEST_TRUE(TEXT("Bag should have map: SoftClass->SoftObject"), GetMapResult.IsValid());
+
+			FPropertyBagMapRef SoftClassToSoftObjectMap = GetMapResult.GetValue();
+			AITEST_EQUAL(TEXT("SoftClassToSoftObjectMap should be empty"), SoftClassToSoftObjectMap.Num(), 0);
+
+			TValueOrError<bool, EPropertyBagResult> ContainsResult = SoftClassToSoftObjectMap.Contains(FSoftClassPath("/Game/Classy"));
+			AITEST_TRUE(TEXT("SoftClassToSoftObjectMap.Contains result should have value"), ContainsResult.HasValue());
+			AITEST_EQUAL(TEXT("SoftClassToSoftObjectMap.Contains result should be false"), ContainsResult.GetValue(), false);
+
+			EPropertyBagResult AddResult = SoftClassToSoftObjectMap.Add(FSoftClassPath("/Game/Classy"), FSoftObjectPath("/Game/Answer"));
+			AITEST_EQUAL(TEXT("SoftClassToSoftObjectMap.Add<FSoftClassPath, FSoftObjectPath> should succeed"), AddResult, EPropertyBagResult::Success);
+			AITEST_EQUAL(TEXT("SoftClassToSoftObjectMap should have 1 pair"), SoftClassToSoftObjectMap.Num(), 1);
+
+			ContainsResult = SoftClassToSoftObjectMap.Contains(FSoftClassPath("/Game/Classy"));
+			AITEST_TRUE(TEXT("SoftClassToSoftObjectMap.Contains result should have value"), ContainsResult.HasValue());
+			AITEST_EQUAL(TEXT("SoftClassToSoftObjectMap.Contains result should be true"), ContainsResult.GetValue(), true);
+
+			UBagTestObject1* TestObject1_Instance1 = NewObject<UBagTestObject1>();
+
+			AddResult = SoftClassToSoftObjectMap.Add(UBagTestObject2::StaticClass(), FSoftObjectPath(TestObject1_Instance1));
+			AITEST_EQUAL(TEXT("SoftClassToSoftObjectMap.Add<FSoftClassPath, FSoftObjectPath> should succeed"), AddResult, EPropertyBagResult::Success);
+			AITEST_EQUAL(TEXT("SoftClassToSoftObjectMap should have 2 pairs"), SoftClassToSoftObjectMap.Num(), 2);
+
+			ContainsResult = SoftClassToSoftObjectMap.Contains(FSoftClassPath(UBagTestObject2::StaticClass()));
+			AITEST_TRUE(TEXT("SoftClassToSoftObjectMap.Contains result should have value"), ContainsResult.HasValue());
+			AITEST_EQUAL(TEXT("SoftClassToSoftObjectMap.Contains result should be true"), ContainsResult.GetValue(), true);
+
+			AddResult = SoftClassToSoftObjectMap.Add(UBagTestObject1::StaticClass(), FSoftObjectPath("/Game/Mismatch"));
+			AITEST_EQUAL(TEXT("SoftClassToSoftObjectMap.Add<UBagTestObject2> should fail with TypeMismatch"), AddResult, EPropertyBagResult::TypeMismatch);
+			AITEST_EQUAL(TEXT("SoftClassToSoftObjectMap should have 2 pairs"), SoftClassToSoftObjectMap.Num(), 2);
+
+			TValueOrError<FSoftObjectPath, EPropertyBagResult> FindAsSoftObjectPathResult = SoftClassToSoftObjectMap.Find<FSoftClassPath, FSoftObjectPath>(FSoftClassPath(UBagTestObject2::StaticClass()));
+			AITEST_TRUE(TEXT("SoftClassToSoftObjectMap.Find<FSoftClassPath, FSoftObjectPath> result should have value"), FindAsSoftObjectPathResult.HasValue());
+			AITEST_EQUAL(TEXT("SoftClassToSoftObjectMap.Find<FSoftClassPath, FSoftObjectPath> result should be expected value"), FindAsSoftObjectPathResult.GetValue(), FSoftObjectPath(TestObject1_Instance1));
+
+			FindAsSoftObjectPathResult = SoftClassToSoftObjectMap.Find<UClass*, FSoftObjectPath>(UBagTestObject2::StaticClass());
+			AITEST_TRUE(TEXT("SoftClassToSoftObjectMap.Find<UClass*, FSoftObjectPath> result should have value"), FindAsSoftObjectPathResult.HasValue());
+			AITEST_EQUAL(TEXT("SoftClassToSoftObjectMap.Find<UClass*, FSoftObjectPath> result should be expected value"), FindAsSoftObjectPathResult.GetValue(), FSoftObjectPath(TestObject1_Instance1));
+
+			TValueOrError<UBagTestObject1*, EPropertyBagResult> FindAsUObjectPtrResult = SoftClassToSoftObjectMap.Find<FSoftClassPath, UBagTestObject1*>(FSoftClassPath(UBagTestObject2::StaticClass()));
+			AITEST_TRUE(TEXT("SoftClassToSoftObjectMap.Find<FSoftClassPath, UBagTestObject1*> result should have value"), FindAsUObjectPtrResult.HasValue());
+			AITEST_EQUAL(TEXT("SoftClassToSoftObjectMap.Find<FSoftClassPath, UBagTestObject1*> result should be expected value"), FindAsUObjectPtrResult.GetValue(), TestObject1_Instance1);
+
+			FindAsUObjectPtrResult = SoftClassToSoftObjectMap.Find<UClass*, UBagTestObject1*>(UBagTestObject2::StaticClass());
+			AITEST_TRUE(TEXT("SoftClassToSoftObjectMap.Find<UClass*, UBagTestObject1*> result should have value"), FindAsUObjectPtrResult.HasValue());
+			AITEST_EQUAL(TEXT("SoftClassToSoftObjectMap.Find<UClass*, UBagTestObject1*> result should be expected value"), FindAsUObjectPtrResult.GetValue(), TestObject1_Instance1);
+
+			EPropertyBagResult RemoveResult = SoftClassToSoftObjectMap.Remove(FSoftClassPath("/Game/Classy"));
+			AITEST_EQUAL(TEXT("SoftClassToSoftObjectMap.Remove should succeed"), RemoveResult, EPropertyBagResult::Success);
+			AITEST_EQUAL(TEXT("SoftClassToSoftObjectMap should have 1 pair"), SoftClassToSoftObjectMap.Num(), 1);
+
+			RemoveResult = SoftClassToSoftObjectMap.Remove(UBagTestObject2::StaticClass());
+			AITEST_EQUAL(TEXT("SoftClassToSoftObjectMap.Remove should succeed"), RemoveResult, EPropertyBagResult::Success);
+			AITEST_EQUAL(TEXT("SoftClassToSoftObjectMap should be empty"), SoftClassToSoftObjectMap.Num(), 0);
+		}
+
+		// map: Struct->SoftClass
+		{
+			TValueOrError<const FPropertyBagMapRef, EPropertyBagResult> GetMapResult = Bag.GetMapRef(StructToSoftClassMapName);
+			AITEST_TRUE(TEXT("Bag should have map: Struct->SoftClass"), GetMapResult.IsValid());
+
+			FPropertyBagMapRef StructToSoftClassMap = GetMapResult.GetValue();
+			AITEST_EQUAL(TEXT("StructToSoftClassMap should be empty"), StructToSoftClassMap.Num(), 0);
+			
+			FTestStructHashable1 Struct1;
+			Struct1.Float = 3.14f;
+
+			TValueOrError<bool, EPropertyBagResult> ContainsResult = StructToSoftClassMap.Contains(Struct1);
+			AITEST_TRUE(TEXT("StructToSoftClassMap.Contains result should have value"), ContainsResult.HasValue());
+			AITEST_EQUAL(TEXT("StructToSoftClassMap.Contains result should be false"), ContainsResult.GetValue(), false);
+
+			EPropertyBagResult AddResult = StructToSoftClassMap.Add(Struct1, FSoftClassPath("/Game/Classy"));
+			AITEST_EQUAL(TEXT("StructToSoftClassMap.Add<FSoftClassPath> should succeed"), AddResult, EPropertyBagResult::Success);
+			AITEST_EQUAL(TEXT("StructToSoftClassMap should have 1 pair"), StructToSoftClassMap.Num(), 1);
+
+			ContainsResult = StructToSoftClassMap.Contains(Struct1);
+			AITEST_TRUE(TEXT("StructToSoftClassMap.Contains result should have value"), ContainsResult.HasValue());
+			AITEST_EQUAL(TEXT("StructToSoftClassMap.Contains result should be true"), ContainsResult.GetValue(), true);
+
+			FTestStructHashable1 Struct2;
+			Struct2.Float = 3.14159f;
+
+			AddResult = StructToSoftClassMap.Add(Struct2, UBagTestObject1::StaticClass());
+			AITEST_EQUAL(TEXT("StructToSoftClassMap.Add<UClass*> should succeed"), AddResult, EPropertyBagResult::Success);
+			AITEST_EQUAL(TEXT("StructToSoftClassMap should have 2 pairs"), StructToSoftClassMap.Num(), 2);
+
+			ContainsResult = StructToSoftClassMap.Contains(Struct2);
+			AITEST_TRUE(TEXT("StructToSoftClassMap.Contains result should have value"), ContainsResult.HasValue());
+			AITEST_EQUAL(TEXT("StructToSoftClassMap.Contains result should be true"), ContainsResult.GetValue(), true);
+
+			FTestStructHashable2 WrongStructType;
+			WrongStructType.Int = 42;
+
+			AddResult = StructToSoftClassMap.Add(WrongStructType, FSoftClassPath("/Game/Mismatch"));
+			AITEST_EQUAL(TEXT("StructToSoftClassMap.Add<FSoftClassPath> should fail with TypeMismatch"), AddResult, EPropertyBagResult::TypeMismatch);
+			AITEST_EQUAL(TEXT("StructToSoftClassMap should have 2 pairs"), StructToSoftClassMap.Num(), 2);
+
+			TValueOrError<UClass*, EPropertyBagResult> FindAsUClassPtrResult = StructToSoftClassMap.Find<FTestStructHashable1, UClass*>(Struct1);
+			AITEST_TRUE(TEXT("StructToSoftClassMap.Find<FTestStructHashable1, UClass*> result should have value"), FindAsUClassPtrResult.HasValue());
+			AITEST_NULL(TEXT("StructToSoftClassMap.Find<FTestStructHashable1, UClass*> result should be null"), FindAsUClassPtrResult.GetValue());
+
+			TValueOrError<FSoftClassPath, EPropertyBagResult> FindAsSoftClassPathResult = StructToSoftClassMap.Find<FTestStructHashable1, FSoftClassPath>(Struct1);
+			AITEST_TRUE(TEXT("StructToSoftClassMap.Find<FTestStructHashable1, FSoftClassPath> result should have value"), FindAsSoftClassPathResult.HasValue());
+			AITEST_EQUAL(TEXT("StructToSoftClassMap.Find<FTestStructHashable1, FSoftClassPath> result should be expected value"), FindAsSoftClassPathResult.GetValue(), FSoftClassPath("/Game/Classy"));
+
+			FindAsUClassPtrResult = StructToSoftClassMap.Find<FTestStructHashable1, UClass*>(Struct2);
+			AITEST_TRUE(TEXT("StructToSoftClassMap.Find<FTestStructHashable1, UClass*> result should have value"), FindAsUClassPtrResult.HasValue());
+			AITEST_EQUAL(TEXT("StructToSoftClassMap.Find<FTestStructHashable1, UClass*> result should be expected value"), FindAsUClassPtrResult.GetValue(), UBagTestObject1::StaticClass());
+
+			FindAsSoftClassPathResult = StructToSoftClassMap.Find<FTestStructHashable1, FSoftClassPath>(Struct2);
+			AITEST_TRUE(TEXT("StructToSoftClassMap.Find<FTestStructHashable1, FSoftClassPath> result should have value"), FindAsSoftClassPathResult.HasValue());
+			AITEST_EQUAL(TEXT("StructToSoftClassMap.Find<FTestStructHashable1, FSoftClassPath> result should be expected value"), FindAsSoftClassPathResult.GetValue(), FSoftClassPath(UBagTestObject1::StaticClass()));
+
+			EPropertyBagResult RemoveResult = StructToSoftClassMap.Remove(Struct2);
+			AITEST_EQUAL(TEXT("StructToSoftClassMap.Remove should succeed"), RemoveResult, EPropertyBagResult::Success);
+			AITEST_EQUAL(TEXT("StructToSoftClassMap should have 1 pair"), StructToSoftClassMap.Num(), 1);
+
+			RemoveResult = StructToSoftClassMap.Remove(Struct1);
+			AITEST_EQUAL(TEXT("StructToSoftClassMap.Remove should succeed"), RemoveResult, EPropertyBagResult::Success);
+			AITEST_EQUAL(TEXT("StructToSoftClassMap should be empty"), StructToSoftClassMap.Num(), 0);
+		}
+
+		// map: Vector->Class
+		{
+			TValueOrError<const FPropertyBagMapRef, EPropertyBagResult> GetMapResult = Bag.GetMapRef(VectorToClassMapName);
+			AITEST_TRUE(TEXT("Bag should have map: Vector->Class"), GetMapResult.IsValid());
+
+			FPropertyBagMapRef VectorToClassMap = GetMapResult.GetValue();
+			AITEST_EQUAL(TEXT("VectorToClassMap should be empty"), VectorToClassMap.Num(), 0);
+
+			TValueOrError<bool, EPropertyBagResult> ContainsResult = VectorToClassMap.Contains(FVector{ 1.0, 2.0, 3.0 });
+			AITEST_TRUE(TEXT("VectorToClassMap.Contains result should have value"), ContainsResult.HasValue());
+			AITEST_EQUAL(TEXT("VectorToClassMap.Contains result should be false"), ContainsResult.GetValue(), false);
+
+			EPropertyBagResult AddResult = VectorToClassMap.Add(FVector{ 1.0, 2.0, 3.0 }, UBagTestObject1::StaticClass());
+			AITEST_EQUAL(TEXT("VectorToClassMap.Add should succeed"), AddResult, EPropertyBagResult::Success);
+			AITEST_EQUAL(TEXT("VectorToClassMap should have 1 pair"), VectorToClassMap.Num(), 1);
+
+			ContainsResult = VectorToClassMap.Contains(FVector{ 1.0, 2.0, 3.0 });
+			AITEST_TRUE(TEXT("VectorToClassMap.Contains result should have value"), ContainsResult.HasValue());
+			AITEST_EQUAL(TEXT("VectorToClassMap.Contains result should be true"), ContainsResult.GetValue(), true);
+
+			AddResult = VectorToClassMap.Add(FVector{ 4.0, 5.0, 6.0 }, UBagTestObject1Derived::StaticClass());
+			AITEST_EQUAL(TEXT("VectorToClassMap.Add should succeed"), AddResult, EPropertyBagResult::Success);
+			AITEST_EQUAL(TEXT("VectorToClassMap should have 2 pairs"), VectorToClassMap.Num(), 2);
+
+			ContainsResult = VectorToClassMap.Contains(FVector{ 4.0, 5.0, 6.0 });
+			AITEST_TRUE(TEXT("VectorToClassMap.Contains result should have value"), ContainsResult.HasValue());
+			AITEST_EQUAL(TEXT("VectorToClassMap.Contains result should be true"), ContainsResult.GetValue(), true);
+
+			TValueOrError<UClass*, EPropertyBagResult> FindResult = VectorToClassMap.Find<FVector, UClass*>(FVector{ 1.0, 2.0, 3.0 });
+			AITEST_TRUE(TEXT("VectorToClassMap.Find result should have value"), FindResult.HasValue());
+			AITEST_EQUAL(TEXT("VectorToClassMap.Find result should be expected value"), FindResult.GetValue(), UBagTestObject1::StaticClass());
+
+			FindResult = VectorToClassMap.Find<FVector, UClass*>(FVector{ 4.0, 5.0, 6.0 });
+			AITEST_TRUE(TEXT("VectorToClassMap.Find result should have value"), FindResult.HasValue());
+			AITEST_EQUAL(TEXT("VectorToClassMap.Find result should be expected value"), FindResult.GetValue(), UBagTestObject1Derived::StaticClass());
+
+			EPropertyBagResult RemoveResult = VectorToClassMap.Remove(FVector{ 1.0, 2.0, 3.0 });
+			AITEST_EQUAL(TEXT("VectorToClassMap.Remove should succeed"), RemoveResult, EPropertyBagResult::Success);
+			AITEST_EQUAL(TEXT("VectorToClassMap should have 1 pair"), VectorToClassMap.Num(), 1);
+
+			RemoveResult = VectorToClassMap.Remove(FVector{ 1.0, 2.0, 3.0 });
+			AITEST_EQUAL(TEXT("VectorToClassMap.Remove should fail"), RemoveResult, EPropertyBagResult::PropertyNotFound);
+			AITEST_EQUAL(TEXT("VectorToClassMap should have 1 pair"), VectorToClassMap.Num(), 1);
+
+			RemoveResult = VectorToClassMap.Remove(FVector{ 4.0, 5.0, 6.0 });
+			AITEST_EQUAL(TEXT("VectorToClassMap.Remove should succeed"), RemoveResult, EPropertyBagResult::Success);
+			AITEST_EQUAL(TEXT("VectorToClassMap should be empty"), VectorToClassMap.Num(), 0);
+		}
+
+		return true;
+	}
+};
+IMPLEMENT_AI_INSTANT_TEST(FTest_Maps, "System.StructUtils.PropertyBag.Maps");
+
+struct FTest_MapsSerialize : FAITestBase
+{
+	virtual bool InstantTest() override
+	{
+		static const FName NameToFloatMapName(TEXT("NameToFloatMap"));
+		static const FName EnumToStructMapName(TEXT("EnumToStructMap"));
+		static const FName VectorToClassMapName(TEXT("VectorToClassMap"));
+
+		FInstancedPropertyBag Bag;
+		Bag.AddProperties({
+			{ NameToFloatMapName, EPropertyBagContainerType::Map, EPropertyBagPropertyType::Float, nullptr, CPF_Edit, EPropertyBagPropertyType::Name },
+			{ EnumToStructMapName, EPropertyBagContainerType::Map, EPropertyBagPropertyType::Struct, FTestStructHashable1::StaticStruct(), CPF_Edit, EPropertyBagPropertyType::Enum, StaticEnum<EPropertyBagTest1>() },
+			{ VectorToClassMapName, EPropertyBagContainerType::Map, EPropertyBagPropertyType::Class, UBagTestObject1::StaticClass(), CPF_Edit, EPropertyBagPropertyType::Struct, TBaseStructure<FVector>::Get() },
+		});
+
+		TArray<uint8> Memory;
+
+		FMemoryWriter Writer(Memory);
+		FObjectAndNameAsStringProxyArchive WriterProxy(Writer, /*bInLoadIfFindFails*/false);
+		const bool bSaveResult = Bag.Serialize(WriterProxy);
+		AITEST_TRUE(TEXT("Saving should succeed"), bSaveResult);
+
+		FInstancedPropertyBag Bag2;
+		FMemoryReader Reader(Memory);
+		FObjectAndNameAsStringProxyArchive ReaderProxy(Reader, /*bInLoadIfFindFails*/true);
+		const bool bLoadResult = Bag2.Serialize(ReaderProxy);
+		AITEST_TRUE(TEXT("Loading to Bag2 should succeed"), bLoadResult);
+		AITEST_EQUAL(TEXT("Bag2 should have 3 items"), Bag2.GetNumPropertiesInBag(), 3);
+
+		{
+			const FPropertyBagPropertyDesc* NameToFloatMapPropertyDesc = Bag2.FindPropertyDescByName(NameToFloatMapName);
+			AITEST_NOT_NULL(TEXT("NameToFloatMapPropertyDesc should be not null"), NameToFloatMapPropertyDesc);
+			AITEST_EQUAL(TEXT("NameToFloatMapPropertyDesc->KeyType should be Name"), NameToFloatMapPropertyDesc->KeyType, EPropertyBagPropertyType::Name);
+			AITEST_NULL(TEXT("NameToFloatMapPropertyDesc->KeyTypeObject should be null"), NameToFloatMapPropertyDesc->KeyTypeObject);
+			AITEST_EQUAL(TEXT("NameToFloatMapPropertyDesc->ValueType should be Float"), NameToFloatMapPropertyDesc->ValueType, EPropertyBagPropertyType::Float);
+			AITEST_NULL(TEXT("NameToFloatMapPropertyDesc->ValueTypeObject should be null"), NameToFloatMapPropertyDesc->ValueTypeObject);
+			AITEST_EQUAL(TEXT("NameToFloatMapPropertyDesc->ContainerTypes should have 1 type"), NameToFloatMapPropertyDesc->ContainerTypes.Num(), 1);
+			AITEST_EQUAL(TEXT("NameToFloatMapPropertyDesc->ContainerTypes.GetFirstContainerType() should be Map"), NameToFloatMapPropertyDesc->ContainerTypes.GetFirstContainerType(), EPropertyBagContainerType::Map);
+		}
+
+		{
+			const FPropertyBagPropertyDesc* EnumToStructMapPropertyDesc = Bag2.FindPropertyDescByName(EnumToStructMapName);
+			AITEST_NOT_NULL(TEXT("EnumToStructMapPropertyDesc should be not null"), EnumToStructMapPropertyDesc);
+			AITEST_EQUAL(TEXT("EnumToStructMapPropertyDesc->KeyType should be Enum"), EnumToStructMapPropertyDesc->KeyType, EPropertyBagPropertyType::Enum);
+			AITEST_EQUAL(TEXT("EnumToStructMapPropertyDesc->KeyTypeObject should be EPropertyBagTest1"), Cast<UEnum>(EnumToStructMapPropertyDesc->KeyTypeObject), StaticEnum<EPropertyBagTest1>());
+			AITEST_EQUAL(TEXT("EnumToStructMapPropertyDesc->ValueType should be Struct"), EnumToStructMapPropertyDesc->ValueType, EPropertyBagPropertyType::Struct);
+			AITEST_EQUAL(TEXT("EnumToStructMapPropertyDesc->ValueTypeObject should be FTestStructHashable1"), Cast<UScriptStruct>(EnumToStructMapPropertyDesc->ValueTypeObject), FTestStructHashable1::StaticStruct());
+			AITEST_EQUAL(TEXT("EnumToStructMapPropertyDesc->ContainerTypes should have 1 type"), EnumToStructMapPropertyDesc->ContainerTypes.Num(), 1);
+			AITEST_EQUAL(TEXT("EnumToStructMapPropertyDesc->ContainerTypes.GetFirstContainerType() should be Map"), EnumToStructMapPropertyDesc->ContainerTypes.GetFirstContainerType(), EPropertyBagContainerType::Map);
+		}
+		
+		{
+			const FPropertyBagPropertyDesc* VectorToClassMapPropertyDesc = Bag2.FindPropertyDescByName(VectorToClassMapName);
+			AITEST_NOT_NULL(TEXT("VectorToClassMapPropertyDesc should be not null"), VectorToClassMapPropertyDesc);
+			AITEST_EQUAL(TEXT("VectorToClassMapPropertyDesc->KeyType should be Struct"), VectorToClassMapPropertyDesc->KeyType, EPropertyBagPropertyType::Struct);
+			AITEST_EQUAL(TEXT("VectorToClassMapPropertyDesc->KeyTypeObject should be FVector"), Cast<UScriptStruct>(VectorToClassMapPropertyDesc->KeyTypeObject), TBaseStructure<FVector>::Get());
+			AITEST_EQUAL(TEXT("VectorToClassMapPropertyDesc->ValueType should be Class"), VectorToClassMapPropertyDesc->ValueType, EPropertyBagPropertyType::Class);
+			AITEST_EQUAL(TEXT("VectorToClassMapPropertyDesc->ValueTypeObject should be UBagTestObject1"), Cast<UClass>(VectorToClassMapPropertyDesc->ValueTypeObject), UBagTestObject1::StaticClass());
+			AITEST_EQUAL(TEXT("VectorToClassMapPropertyDesc->ContainerTypes should have 1 type"), VectorToClassMapPropertyDesc->ContainerTypes.Num(), 1);
+			AITEST_EQUAL(TEXT("VectorToClassMapPropertyDesc->ContainerTypes.GetFirstContainerType() should be Map"), VectorToClassMapPropertyDesc->ContainerTypes.GetFirstContainerType(), EPropertyBagContainerType::Map);
+		}
+
+		return true;
+	}
+};
+IMPLEMENT_AI_INSTANT_TEST(FTest_MapsSerialize, "System.StructUtils.PropertyBag.MapsSerialize");
+
+struct FTest_MapsNestedArray : FAITestBase
+{
+	virtual bool InstantTest() override
+	{
+		static const FName NameToIntArrayMapName(TEXT("NameToIntArrayMap"));
+		static const FName NestedArrayName(TEXT("NestedIntArray"));
+		static const TArray<int32> TestArray = {42, 13, 57};
+
+		FInstancedPropertyBag Bag;
+
+		{
+			EPropertyBagAlterationResult AlterationResult = Bag.AddMapProperty(NameToIntArrayMapName, EPropertyBagPropertyType::Name, EPropertyBagPropertyType::Int32, nullptr, nullptr, EPropertyBagContainerType::Array);
+			AITEST_EQUAL(TEXT("Bag.AddMapProperty should succeed"), AlterationResult, EPropertyBagAlterationResult::Success);
+		}
+
+		TValueOrError<const FPropertyBagMapRef, EPropertyBagResult> GetMapResult = Bag.GetMapRef(NameToIntArrayMapName);
+		AITEST_TRUE(TEXT("Bag should have map: Name->IntArray"), GetMapResult.IsValid());
+
+		FPropertyBagMapRef NameToIntArrayMap = GetMapResult.GetValue();
+		AITEST_EQUAL(TEXT("NameToIntArrayMap should be empty"), NameToIntArrayMap.Num(), 0);
+
+		// Verify that we can't add values directly to the map
+		{
+			EPropertyBagResult AddResult = NameToIntArrayMap.Add(FName("NotAnArray"), 42);
+			AITEST_EQUAL(TEXT("NameToIntArrayMap.Add<FName,int> should fail with TypeMismatch"), AddResult, EPropertyBagResult::TypeMismatch);
+			AITEST_EQUAL(TEXT("NameToIntArrayMap should be empty"), NameToIntArrayMap.Num(), 0);
+
+			AddResult = NameToIntArrayMap.Add(FName("AnIntArray"), TArray<int32>{ 1, 2, 3 });
+			AITEST_EQUAL(TEXT("NameToIntArrayMap.Add<FName,TArray<int32>> should fail with TypeMismatch"), AddResult, EPropertyBagResult::TypeMismatch);
+			AITEST_EQUAL(TEXT("NameToIntArrayMap should be empty"), NameToIntArrayMap.Num(), 0);
+		}
+		
+		// Reserve a pair for a given key
+		{
+			EPropertyBagResult AddResult = NameToIntArrayMap.Add(NestedArrayName);
+			AITEST_EQUAL(TEXT("NameToIntArrayMap.Add<FName> should succeed"), AddResult, EPropertyBagResult::Success);
+			AITEST_EQUAL(TEXT("NameToIntArrayMap should have 1 pair"), NameToIntArrayMap.Num(), 1);
+			
+			AddResult = NameToIntArrayMap.Add(NestedArrayName);
+			AITEST_EQUAL(TEXT("NameToIntArrayMap.Add<FName> should fail with DuplicatedValue"), AddResult, EPropertyBagResult::DuplicatedValue);
+			AITEST_EQUAL(TEXT("NameToIntArrayMap should have 1 pair"), NameToIntArrayMap.Num(), 1);
+		}
+
+		// Verify that we can't get helpers for other nested container types
+		{
+			TValueOrError<FPropertyBagSetRef, EPropertyBagResult> GetMutableNestedSetResult = NameToIntArrayMap.GetMutableNestedSetRef(NestedArrayName);
+			AITEST_TRUE(TEXT("NameToIntArrayMap.GetMutableNestedSetRef<FName>"), GetMutableNestedSetResult.HasError());
+			AITEST_EQUAL(TEXT("NameToIntArrayMap.GetMutableNestedSetRef<FName> should fail with TypeMismatch"), GetMutableNestedSetResult.GetError(), EPropertyBagResult::TypeMismatch);
+			
+			TValueOrError<FPropertyBagMapRef, EPropertyBagResult> GetMutableNestedMapResult = NameToIntArrayMap.GetMutableNestedMapRef(NestedArrayName);
+			AITEST_TRUE(TEXT("NameToIntArrayMap.GetMutableNestedMapRef<FName>"), GetMutableNestedMapResult.HasError());
+			AITEST_EQUAL(TEXT("NameToIntArrayMap.GetMutableNestedMapRef<FName> should fail with TypeMismatch"), GetMutableNestedMapResult.GetError(), EPropertyBagResult::TypeMismatch);
+		}
+
+		// Set Nested Array values using mutable helper interface
+		{
+			TValueOrError<FPropertyBagArrayRef, EPropertyBagResult> GetMutableNestedArrayResult = NameToIntArrayMap.GetMutableNestedArrayRef(NestedArrayName);
+			AITEST_TRUE(TEXT("NameToIntArrayMap.GetMutableNestedArrayRef<FName> should succeed."), GetMutableNestedArrayResult.HasValue());
+			
+			FPropertyBagArrayRef MutableNestedIntArray = GetMutableNestedArrayResult.GetValue();
+			AITEST_EQUAL(TEXT("MutableNestedIntArray should be empty"), MutableNestedIntArray.Num(), 0);
+
+			MutableNestedIntArray.AddUninitializedValues(TestArray.Num());
+
+			for (int32 i = 0; i < TestArray.Num(); i++)
+			{
+				EPropertyBagResult SetValueResult = MutableNestedIntArray.SetValueInt32(i, TestArray[i]);
+				AITEST_EQUAL(TEXT("MutableNestedIntArray.SetValueInt32 should succeed"), SetValueResult, EPropertyBagResult::Success);
+			}
+
+			AITEST_EQUAL(TEXT("MutableNestedIntArray.Num should match expected count"), MutableNestedIntArray.Num(), TestArray.Num());
+		}
+
+		// Get Nested Array values using const helper interface
+		{
+			TValueOrError<const FPropertyBagArrayRef, EPropertyBagResult> GetNestedArrayResult = NameToIntArrayMap.GetNestedArrayRef(NestedArrayName);
+			AITEST_TRUE(TEXT("NameToIntArrayMap.GetNestedArrayResult<FName> should succeed."), GetNestedArrayResult.HasValue());
+
+			const FPropertyBagArrayRef& NestedIntArray = GetNestedArrayResult.GetValue();
+			AITEST_EQUAL(TEXT("NestedIntArray.Num should match expected count"), NestedIntArray.Num(), TestArray.Num());
+
+			for (int32 i = 0; i < NestedIntArray.Num(); i++)
+			{
+				const TValueOrError<int32, EPropertyBagResult> GetValueResult = NestedIntArray.GetValueInt32(i);
+				AITEST_TRUE(TEXT("NestedIntArray.GetValueInt32 should succeed"), GetValueResult.HasValue());
+				AITEST_EQUAL(TEXT("NestedIntArray.GetValueInt32 should match expected value"), GetValueResult.GetValue(), TestArray[i]);
+			}
+		}
+
+		return true;
+	}
+};
+IMPLEMENT_AI_INSTANT_TEST(FTest_MapsNestedArray, "System.StructUtils.PropertyBag.MapsNestedArray");
+
+struct FTest_MapsNestedSet : FAITestBase
+{
+	virtual bool InstantTest() override
+	{
+		static const FName NameToEnumSetMapName(TEXT("NameToEnumSetMap"));
+		static const FName NestedSetName(TEXT("NestedEnumSet"));
+		static const TSet<EPropertyBagTest1> TestSet = { EPropertyBagTest1::Foo, EPropertyBagTest1::Bar, EPropertyBagTest1::Baz };
+
+		FInstancedPropertyBag Bag;
+
+		{
+			EPropertyBagAlterationResult AlterationResult = Bag.AddMapProperty(NameToEnumSetMapName, EPropertyBagPropertyType::Name, EPropertyBagPropertyType::Enum, nullptr, StaticEnum<EPropertyBagTest1>(), EPropertyBagContainerType::Set);
+			AITEST_EQUAL(TEXT("Bag.AddMapProperty should succeed"), AlterationResult, EPropertyBagAlterationResult::Success);
+		}
+
+		TValueOrError<const FPropertyBagMapRef, EPropertyBagResult> GetMapResult = Bag.GetMapRef(NameToEnumSetMapName);
+		AITEST_TRUE(TEXT("Bag should have map: Name->EnumSet"), GetMapResult.IsValid());
+
+		FPropertyBagMapRef NameToEnumSetMap = GetMapResult.GetValue();
+		AITEST_EQUAL(TEXT("NameToEnumSetMap should be empty"), NameToEnumSetMap.Num(), 0);
+
+		// Verify that we can't add values directly to the map
+		{
+			EPropertyBagResult AddResult = NameToEnumSetMap.Add(FName("NotASet"), EPropertyBagTest1::Foo);
+			AITEST_EQUAL(TEXT("NameToEnumSetMap.Add<FName,Enum> should fail with TypeMismatch"), AddResult, EPropertyBagResult::TypeMismatch);
+			AITEST_EQUAL(TEXT("NameToEnumSetMap should be empty"), NameToEnumSetMap.Num(), 0);
+
+			AddResult = NameToEnumSetMap.Add(FName("AnEnumSet"), TSet<EPropertyBagTest1>{ EPropertyBagTest1::Foo });
+			AITEST_EQUAL(TEXT("NameToEnumSetMap.Add<FName,TSet<Enum>> should fail with TypeMismatch"), AddResult, EPropertyBagResult::TypeMismatch);
+			AITEST_EQUAL(TEXT("NameToEnumSetMap should be empty"), NameToEnumSetMap.Num(), 0);
+		}
+
+		// Reserve a pair for a given key
+		{
+			EPropertyBagResult AddResult = NameToEnumSetMap.Add(NestedSetName);
+			AITEST_EQUAL(TEXT("NameToEnumSetMap.Add<FName> should succeed"), AddResult, EPropertyBagResult::Success);
+			AITEST_EQUAL(TEXT("NameToEnumSetMap should have 1 pair"), NameToEnumSetMap.Num(), 1);
+
+			AddResult = NameToEnumSetMap.Add(NestedSetName);
+			AITEST_EQUAL(TEXT("NameToEnumSetMap.Add<FName> should fail with DuplicatedValue"), AddResult, EPropertyBagResult::DuplicatedValue);
+			AITEST_EQUAL(TEXT("NameToEnumSetMap should have 1 pair"), NameToEnumSetMap.Num(), 1);
+		}
+
+		// Verify that we can't get helpers for other nested container types
+		{
+			TValueOrError<FPropertyBagArrayRef, EPropertyBagResult> GetMutableNestedArrayResult = NameToEnumSetMap.GetMutableNestedArrayRef(NestedSetName);
+			AITEST_TRUE(TEXT("NameToEnumSetMap.GetMutableNestedArrayRef<FName>"), GetMutableNestedArrayResult.HasError());
+			AITEST_EQUAL(TEXT("NameToEnumSetMap.GetMutableNestedArrayRef<FName> should fail with TypeMismatch"), GetMutableNestedArrayResult.GetError(), EPropertyBagResult::TypeMismatch);
+
+			TValueOrError<FPropertyBagMapRef, EPropertyBagResult> GetMutableNestedMapResult = NameToEnumSetMap.GetMutableNestedMapRef(NestedSetName);
+			AITEST_TRUE(TEXT("NameToEnumSetMap.GetMutableNestedMapRef<FName>"), GetMutableNestedMapResult.HasError());
+			AITEST_EQUAL(TEXT("NameToEnumSetMap.GetMutableNestedMapRef<FName> should fail with TypeMismatch"), GetMutableNestedMapResult.GetError(), EPropertyBagResult::TypeMismatch);
+		}
+
+		// Add Nested Set values using mutable helper interface
+		{
+			TValueOrError<FPropertyBagSetRef, EPropertyBagResult> GetMutableNestedSetResult = NameToEnumSetMap.GetMutableNestedSetRef(NestedSetName);
+			AITEST_TRUE(TEXT("NameToEnumSetMap.GetMutableNestedSetRef<FName> should succeed."), GetMutableNestedSetResult.HasValue());
+
+			FPropertyBagSetRef MutableNestedEnumSet = GetMutableNestedSetResult.GetValue();
+			AITEST_EQUAL(TEXT("MutableNestedEnumSet should be empty"), MutableNestedEnumSet.Num(), 0);
+
+			for (EPropertyBagTest1 TestEnum : TestSet)
+			{
+				EPropertyBagResult AddValueResult = MutableNestedEnumSet.AddValueEnum(TestEnum);
+				AITEST_EQUAL(TEXT("MutableNestedEnumSet.AddValueEnum should succeed"), AddValueResult, EPropertyBagResult::Success);
+			}
+
+			AITEST_EQUAL(TEXT("MutableNestedEnumSet.Num should match expected count"), MutableNestedEnumSet.Num(), TestSet.Num());
+		}
+
+		// Check Nested Set values using const helper interface
+		{
+			TValueOrError<const FPropertyBagSetRef, EPropertyBagResult> GetNestedSetResult = NameToEnumSetMap.GetNestedSetRef(NestedSetName);
+			AITEST_TRUE(TEXT("NameToEnumSetMap.GetNestedSetRef<FName> should succeed."), GetNestedSetResult.HasValue());
+			
+			const FPropertyBagSetRef& NestedEnumSet = GetNestedSetResult.GetValue();
+			AITEST_EQUAL(TEXT("NestedEnumSet.Num should match expected count"), NestedEnumSet.Num(), TestSet.Num());
+
+			for (EPropertyBagTest1 TestEnum : TestSet)
+			{
+				TValueOrError<bool, EPropertyBagResult> ContainsResult = NestedEnumSet.Contains(TestEnum);
+				AITEST_TRUE(TEXT("NestedEnumSet.Contains should have value"), ContainsResult.HasValue());
+				AITEST_TRUE(TEXT("NestedEnumSet.Contains should succeed."), ContainsResult.GetValue());
+			}
+		}
+
+		return true;
+	}
+};
+IMPLEMENT_AI_INSTANT_TEST(FTest_MapsNestedSet, "System.StructUtils.PropertyBag.MapsNestedSet");
+
+struct FTest_MapsNestedMap : FAITestBase
+{
+	virtual bool InstantTest() override
+	{
+		static const FName NameToVectorMapMapName(TEXT("NameToVectorMapMap"));
+		static const FName NestedMapName(TEXT("NestedVectorMap"));
+		static const TMap<FName, FVector> TestMap = { { FName("OneTwoThree"), FVector{1, 2, 3} }, { FName("FourFiveSix"), FVector{4, 5, 6} } };
+
+		FInstancedPropertyBag Bag;
+
+		{
+			EPropertyBagAlterationResult AlterationResult = Bag.AddMapProperty(NameToVectorMapMapName, EPropertyBagPropertyType::Name, EPropertyBagPropertyType::Struct, nullptr, TBaseStructure<FVector>::Get(), EPropertyBagContainerType::Map);
+			AITEST_EQUAL(TEXT("Bag.AddMapProperty should succeed"), AlterationResult, EPropertyBagAlterationResult::Success);
+		}
+
+		TValueOrError<const FPropertyBagMapRef, EPropertyBagResult> GetMapResult = Bag.GetMapRef(NameToVectorMapMapName);
+		AITEST_TRUE(TEXT("Bag should have map: Name->VectorMap"), GetMapResult.IsValid());
+
+		FPropertyBagMapRef NameToVectorMapMap = GetMapResult.GetValue();
+		AITEST_EQUAL(TEXT("NameToVectorMapMap should be empty"), NameToVectorMapMap.Num(), 0);
+
+		// Verify that we can't add values directly to the map
+		{
+			EPropertyBagResult AddResult = NameToVectorMapMap.Add(FName("NotAMap"), FVector{ 1, 2, 3});
+			AITEST_EQUAL(TEXT("NameToVectorMapMap.Add<FName,FVector> should fail with TypeMismatch"), AddResult, EPropertyBagResult::TypeMismatch);
+			AITEST_EQUAL(TEXT("NameToVectorMapMap should be empty"), NameToVectorMapMap.Num(), 0);
+
+			AddResult = NameToVectorMapMap.Add(FName("AVectorMap"), TMap<FName, FVector>{ { FName("OneTwoThree"), FVector{ 1, 2, 3 } } });
+			AITEST_EQUAL(TEXT("NameToVectorMapMap.Add<FName,TMap<FName, FVector>> should fail with TypeMismatch"), AddResult, EPropertyBagResult::TypeMismatch);
+			AITEST_EQUAL(TEXT("NameToVectorMapMap should be empty"), NameToVectorMapMap.Num(), 0);
+		}
+
+		// Reserve a pair for a given key
+		{
+			EPropertyBagResult AddResult = NameToVectorMapMap.Add(NestedMapName);
+			AITEST_EQUAL(TEXT("NameToVectorMapMap.Add<FName> should succeed"), AddResult, EPropertyBagResult::Success);
+			AITEST_EQUAL(TEXT("NameToVectorMapMap should have 1 pair"), NameToVectorMapMap.Num(), 1);
+
+			AddResult = NameToVectorMapMap.Add(NestedMapName);
+			AITEST_EQUAL(TEXT("NameToVectorMapMap.Add<FName> should fail with DuplicatedValue"), AddResult, EPropertyBagResult::DuplicatedValue);
+			AITEST_EQUAL(TEXT("NameToVectorMapMap should have 1 pair"), NameToVectorMapMap.Num(), 1);
+		}
+
+		// Verify that we can't get helpers for other nested container types
+		{
+			TValueOrError<FPropertyBagArrayRef, EPropertyBagResult> GetMutableNestedArrayResult = NameToVectorMapMap.GetMutableNestedArrayRef(NestedMapName);
+			AITEST_TRUE(TEXT("NameToVectorMapMap.GetMutableNestedArrayRef<FName>"), GetMutableNestedArrayResult.HasError());
+			AITEST_EQUAL(TEXT("NameToVectorMapMap.GetMutableNestedArrayRef<FName> should fail with TypeMismatch"), GetMutableNestedArrayResult.GetError(), EPropertyBagResult::TypeMismatch);
+
+			TValueOrError<FPropertyBagSetRef, EPropertyBagResult> GetMutableNestedSetResult = NameToVectorMapMap.GetMutableNestedSetRef(NestedMapName);
+			AITEST_TRUE(TEXT("NameToVectorMapMap.GetMutableNestedSetRef<FName>"), GetMutableNestedSetResult.HasError());
+			AITEST_EQUAL(TEXT("NameToVectorMapMap.GetMutableNestedSetRef<FName> should fail with TypeMismatch"), GetMutableNestedSetResult.GetError(), EPropertyBagResult::TypeMismatch);
+		}
+
+		// Add Nested Map values using mutable helper interface
+		{
+			TValueOrError<FPropertyBagMapRef, EPropertyBagResult> GetMutableNestedMapResult = NameToVectorMapMap.GetMutableNestedMapRef(NestedMapName);
+			AITEST_TRUE(TEXT("NameToVectorMapMap.GetMutableNestedMapRef<FName> should succeed."), GetMutableNestedMapResult.HasValue());
+
+			FPropertyBagMapRef MutableNestedVectorMap = GetMutableNestedMapResult.GetValue();
+			AITEST_EQUAL(TEXT("MutableNestedVectorMap should be empty"), MutableNestedVectorMap.Num(), 0);
+
+			for (const TPair<FName, FVector>& Pair : TestMap)
+			{
+				EPropertyBagResult AddValueResult = MutableNestedVectorMap.Add(Pair.Key, Pair.Value);
+				AITEST_EQUAL(TEXT("MutableNestedVectorMap.Add should succeed"), AddValueResult, EPropertyBagResult::Success);
+			}
+
+			AITEST_EQUAL(TEXT("MutableNestedVectorMap.Num should match expected count"), MutableNestedVectorMap.Num(), TestMap.Num());
+		}
+
+		// Check Nested Map values using const helper interface
+		{
+			TValueOrError<const FPropertyBagMapRef, EPropertyBagResult> GetNestedMapResult = NameToVectorMapMap.GetNestedMapRef(NestedMapName);
+			AITEST_TRUE(TEXT("NameToVectorMapMap.GetNestedMapRef<FName> should succeed."), GetNestedMapResult.HasValue());
+
+			const FPropertyBagMapRef& NestedVectorMap = GetNestedMapResult.GetValue();
+			AITEST_EQUAL(TEXT("NestedVectorMap.Num should match expected count"), NestedVectorMap.Num(), TestMap.Num());
+
+			for (const TPair<FName, FVector>& Pair : TestMap)
+			{
+				TValueOrError<bool, EPropertyBagResult> ContainsResult = NestedVectorMap.Contains(Pair.Key);
+				AITEST_TRUE(TEXT("NestedVectorMap.Contains should have value"), ContainsResult.HasValue());
+				AITEST_TRUE(TEXT("NestedVectorMap.Contains should succeed."), ContainsResult.GetValue());
+
+				TValueOrError<FVector, EPropertyBagResult> FindResult = NestedVectorMap.Find<FName, FVector>(Pair.Key);
+				AITEST_TRUE(TEXT("NestedVectorMap.Find should have value"), FindResult.HasValue());
+				AITEST_EQUAL(TEXT("NestedVectorMap.Find should match expected value"), FindResult.GetValue(), Pair.Value);
+			}
+		}
+
+		return true;
+	}
+};
+IMPLEMENT_AI_INSTANT_TEST(FTest_MapsNestedMap, "System.StructUtils.PropertyBag.MapsNestedMap");
+
 struct FTest_SameBag : FAITestBase
 {
 	virtual bool InstantTest() override
diff --git a/Engine/Source/Developer/StructUtilsTestSuite/Private/StructUtilsTestTypes.cpp b/Engine/Source/Developer/StructUtilsTestSuite/Private/StructUtilsTestTypes.cpp
--- a/Engine/Source/Developer/StructUtilsTestSuite/Private/StructUtilsTestTypes.cpp
+++ b/Engine/Source/Developer/StructUtilsTestSuite/Private/StructUtilsTestTypes.cpp
@@ -8,3 +8,8 @@ uint32 GetTypeHash(const FTestStructHashable1& other)
 {
 	return FCrc::MemCrc32(&other, sizeof(FTestStructHashable1));
 }
+
+uint32 GetTypeHash(const FTestStructHashable2& other)
+{
+	return FCrc::MemCrc32(&other, sizeof(FTestStructHashable2));
+}
diff --git a/Engine/Source/Developer/StructUtilsTestSuite/Public/StructUtilsTestTypes.h b/Engine/Source/Developer/StructUtilsTestSuite/Public/StructUtilsTestTypes.h
--- a/Engine/Source/Developer/StructUtilsTestSuite/Public/StructUtilsTestTypes.h
+++ b/Engine/Source/Developer/StructUtilsTestSuite/Public/StructUtilsTestTypes.h
@@ -157,6 +157,19 @@ struct FTestStructHashable1 : public FTestStructSimpleBase
 
 uint32 GetTypeHash(const FTestStructHashable1& other);
 
+USTRUCT()
+struct FTestStructHashable2 : public FTestStructSimpleBase
+{
+	GENERATED_BODY()
+
+	FTestStructHashable2() = default;
+
+	UPROPERTY()
+	int Int = 0;
+};
+
+uint32 GetTypeHash(const FTestStructHashable2& other);
+
 USTRUCT()
 struct FTestStructNonTrivialDestructor
 {
@@ -193,11 +206,18 @@ class UBagTestObject1Derived : public UBagTestObject1
 	GENERATED_BODY()
 };
 
+UCLASS()
+class UBagTestObject1Derived2 : public UBagTestObject1
+{
+	GENERATED_BODY()
+};
+
 UENUM()
 enum class EPropertyBagTest1 : uint8
 {
 	Foo,
 	Bar,
+	Baz,
 };
 
 UENUM()
diff --git a/Engine/Source/Editor/StructUtilsEditor/Private/Customizations/PropertyBagDetails.cpp b/Engine/Source/Editor/StructUtilsEditor/Private/Customizations/PropertyBagDetails.cpp
--- a/Engine/Source/Editor/StructUtilsEditor/Private/Customizations/PropertyBagDetails.cpp
+++ b/Engine/Source/Editor/StructUtilsEditor/Private/Customizations/PropertyBagDetails.cpp
@@ -54,102 +54,127 @@ void SetPropertyDescFromPin(FPropertyBagPropertyDesc& Desc, const FEdGraphPinTyp
 		Desc.ContainerTypes.Add(EPropertyBagContainerType::Set);
 		break;
 	case EPinContainerType::Map:
-		ensureMsgf(false, TEXT("Unsuported container type [Map] "));
+		Desc.ContainerTypes.Add(EPropertyBagContainerType::Map);
 		break;
 	default:
 		break;
 	}
-
-	// Value type
-	if (PinType.PinCategory == UEdGraphSchema_K2::PC_Boolean)
-	{
-		Desc.ValueType = EPropertyBagPropertyType::Bool;
-		Desc.ValueTypeObject = nullptr;
-	}
-	else if (PinType.PinCategory == UEdGraphSchema_K2::PC_Byte)
+	
+	auto SetDescTypeFromPinType = [](
+		const FName& Category,
+		const FName& SubCategory,
+		const TWeakObjectPtr<UObject>& SubCategoryObject,
+		EPropertyBagPropertyType& OutType,
+		TObjectPtr<const UObject>& OutTypeObject)
 	{
-		if (UEnum* Enum = Cast<UEnum>(PinType.PinSubCategoryObject))
+		if (Category == UEdGraphSchema_K2::PC_Boolean)
+		{
+			OutType = EPropertyBagPropertyType::Bool;
+			OutTypeObject = nullptr;
+		}
+		else if (Category == UEdGraphSchema_K2::PC_Byte)
+		{
+			if (UEnum* Enum = Cast<UEnum>(SubCategoryObject))
+			{
+				OutType = EPropertyBagPropertyType::Enum;
+				OutTypeObject = SubCategoryObject.Get();
+			}
+			else
+			{
+				OutType = EPropertyBagPropertyType::Byte;
+				OutTypeObject = nullptr;
+			}
+		}
+		else if (Category == UEdGraphSchema_K2::PC_Int)
+		{
+			OutType = EPropertyBagPropertyType::Int32;
+			OutTypeObject = nullptr;
+		}
+		else if (Category == UEdGraphSchema_K2::PC_Int64)
+		{
+			OutType = EPropertyBagPropertyType::Int64;
+			OutTypeObject = nullptr;
+		}
+		else if (Category == UEdGraphSchema_K2::PC_Real)
+		{
+			if (SubCategory == UEdGraphSchema_K2::PC_Float)
+			{
+				OutType = EPropertyBagPropertyType::Float;
+				OutTypeObject = nullptr;
+			}
+			else if (SubCategory == UEdGraphSchema_K2::PC_Double)
+			{
+				OutType = EPropertyBagPropertyType::Double;
+				OutTypeObject = nullptr;
+			}
+		}
+		else if (Category == UEdGraphSchema_K2::PC_Name)
+		{
+			OutType = EPropertyBagPropertyType::Name;
+			OutTypeObject = nullptr;
+		}
+		else if (Category == UEdGraphSchema_K2::PC_String)
 		{
-			Desc.ValueType = EPropertyBagPropertyType::Enum;
-			Desc.ValueTypeObject = PinType.PinSubCategoryObject.Get();
+			OutType = EPropertyBagPropertyType::String;
+			OutTypeObject = nullptr;
+		}
+		else if (Category == UEdGraphSchema_K2::PC_Text)
+		{
+			OutType = EPropertyBagPropertyType::Text;
+			OutTypeObject = nullptr;
+		}
+		else if (Category == UEdGraphSchema_K2::PC_Enum)
+		{
+			OutType = EPropertyBagPropertyType::Enum;
+			OutTypeObject = SubCategoryObject.Get();
+		}
+		else if (Category == UEdGraphSchema_K2::PC_Struct)
+		{
+			OutType = EPropertyBagPropertyType::Struct;
+			OutTypeObject = SubCategoryObject.Get();
+		}
+		else if (Category == UEdGraphSchema_K2::PC_Object)
+		{
+			OutType = EPropertyBagPropertyType::Object;
+			OutTypeObject = SubCategoryObject.Get();
+		}
+		else if (Category == UEdGraphSchema_K2::PC_SoftObject)
+		{
+			OutType = EPropertyBagPropertyType::SoftObject;
+			OutTypeObject = SubCategoryObject.Get();
+		}
+		else if (Category == UEdGraphSchema_K2::PC_Class)
+		{
+			OutType = EPropertyBagPropertyType::Class;
+			OutTypeObject = SubCategoryObject.Get();
+		}
+		else if (Category == UEdGraphSchema_K2::PC_SoftClass)
+		{
+			OutType = EPropertyBagPropertyType::SoftClass;
+			OutTypeObject = SubCategoryObject.Get();
 		}
 		else
 		{
-			Desc.ValueType = EPropertyBagPropertyType::Byte;
-			Desc.ValueTypeObject = nullptr;
+			ensureMsgf(false, TEXT("Unhandled pin category %s"), *Category.ToString());
 		}
-	}
-	else if (PinType.PinCategory == UEdGraphSchema_K2::PC_Int)
-	{
-		Desc.ValueType = EPropertyBagPropertyType::Int32;
-		Desc.ValueTypeObject = nullptr;
-	}
-	else if (PinType.PinCategory == UEdGraphSchema_K2::PC_Int64)
-	{
-		Desc.ValueType = EPropertyBagPropertyType::Int64;
-		Desc.ValueTypeObject = nullptr;
-	}
-	else if (PinType.PinCategory == UEdGraphSchema_K2::PC_Real)
+	};
+
+	if (PinType.ContainerType == EPinContainerType::Map)
 	{
-		if (PinType.PinSubCategory == UEdGraphSchema_K2::PC_Float)
+		SetDescTypeFromPinType(PinType.PinCategory, PinType.PinSubCategory, PinType.PinSubCategoryObject, Desc.KeyType, Desc.KeyTypeObject);
+		
+		if (PinType.PinValueType.TerminalCategory == NAME_None)
 		{
-			Desc.ValueType = EPropertyBagPropertyType::Float;
-			Desc.ValueTypeObject = nullptr;
+			SetDescTypeFromPinType(PinType.PinCategory, PinType.PinSubCategory, PinType.PinSubCategoryObject, Desc.ValueType, Desc.ValueTypeObject);
 		}
-		else if (PinType.PinSubCategory == UEdGraphSchema_K2::PC_Double)
+		else
 		{
-			Desc.ValueType = EPropertyBagPropertyType::Double;
-			Desc.ValueTypeObject = nullptr;
+			SetDescTypeFromPinType(PinType.PinValueType.TerminalCategory, PinType.PinValueType.TerminalSubCategory, PinType.PinValueType.TerminalSubCategoryObject, Desc.ValueType, Desc.ValueTypeObject);
 		}
 	}
-	else if (PinType.PinCategory == UEdGraphSchema_K2::PC_Name)
-	{
-		Desc.ValueType = EPropertyBagPropertyType::Name;
-		Desc.ValueTypeObject = nullptr;
-	}
-	else if (PinType.PinCategory == UEdGraphSchema_K2::PC_String)
-	{
-		Desc.ValueType = EPropertyBagPropertyType::String;
-		Desc.ValueTypeObject = nullptr;
-	}
-	else if (PinType.PinCategory == UEdGraphSchema_K2::PC_Text)
-	{
-		Desc.ValueType = EPropertyBagPropertyType::Text;
-		Desc.ValueTypeObject = nullptr;
-	}
-	else if (PinType.PinCategory == UEdGraphSchema_K2::PC_Enum)
-	{
-		Desc.ValueType = EPropertyBagPropertyType::Enum;
-		Desc.ValueTypeObject = PinType.PinSubCategoryObject.Get();
-	}
-	else if (PinType.PinCategory == UEdGraphSchema_K2::PC_Struct)
-	{
-		Desc.ValueType = EPropertyBagPropertyType::Struct;
-		Desc.ValueTypeObject = PinType.PinSubCategoryObject.Get();
-	}
-	else if (PinType.PinCategory == UEdGraphSchema_K2::PC_Object)
-	{
-		Desc.ValueType = EPropertyBagPropertyType::Object;
-		Desc.ValueTypeObject = PinType.PinSubCategoryObject.Get();
-	}
-	else if (PinType.PinCategory == UEdGraphSchema_K2::PC_SoftObject)
-	{
-		Desc.ValueType = EPropertyBagPropertyType::SoftObject;
-		Desc.ValueTypeObject = PinType.PinSubCategoryObject.Get();
-	}
-	else if (PinType.PinCategory == UEdGraphSchema_K2::PC_Class)
-	{
-		Desc.ValueType = EPropertyBagPropertyType::Class;
-		Desc.ValueTypeObject = PinType.PinSubCategoryObject.Get();
-	}
-	else if (PinType.PinCategory == UEdGraphSchema_K2::PC_SoftClass)
-	{
-		Desc.ValueType = EPropertyBagPropertyType::SoftClass;
-		Desc.ValueTypeObject = PinType.PinSubCategoryObject.Get();
-	}
 	else
 	{
-		ensureMsgf(false, TEXT("Unhandled pin category %s"), *PinType.PinCategory.ToString());
+		SetDescTypeFromPinType(PinType.PinCategory, PinType.PinSubCategory, PinType.PinSubCategoryObject, Desc.ValueType, Desc.ValueTypeObject);
 	}
 }
 
@@ -175,76 +200,96 @@ FEdGraphPinType GetPropertyDescAsPin(const FPropertyBagPropertyDesc& Desc)
 	case EPropertyBagContainerType::Set:
 		PinType.ContainerType = EPinContainerType::Set;
 		break;
+	case EPropertyBagContainerType::Map:
+		PinType.ContainerType = EPinContainerType::Map;
+		break;
 	default:
 		PinType.ContainerType = EPinContainerType::None;
 	}
 
-	// Value type
-	switch (Desc.ValueType)
+	auto SetPinTypeFromDescType = [](
+		const EPropertyBagPropertyType& Type,
+		const TObjectPtr<const UObject>& TypeObject,
+		FName& OutCategory,
+		FName& OutSubCategory,
+		TWeakObjectPtr<UObject>& OutSubCategoryObject)
 	{
-	case EPropertyBagPropertyType::Bool:
-		PinType.PinCategory = UEdGraphSchema_K2::PC_Boolean;
-		break;
-	case EPropertyBagPropertyType::Byte:
-		PinType.PinCategory = UEdGraphSchema_K2::PC_Byte;
-		break;
-	case EPropertyBagPropertyType::Int32:
-		PinType.PinCategory = UEdGraphSchema_K2::PC_Int;
-		break;
-	case EPropertyBagPropertyType::Int64:
-		PinType.PinCategory = UEdGraphSchema_K2::PC_Int64;
-		break;
-	case EPropertyBagPropertyType::Float:
-		PinType.PinCategory = UEdGraphSchema_K2::PC_Real;
-		PinType.PinSubCategory = UEdGraphSchema_K2::PC_Float;
-		break;
-	case EPropertyBagPropertyType::Double:
-		PinType.PinCategory = UEdGraphSchema_K2::PC_Real;
-		PinType.PinSubCategory = UEdGraphSchema_K2::PC_Double;
-		break;
-	case EPropertyBagPropertyType::Name:
-		PinType.PinCategory = UEdGraphSchema_K2::PC_Name;
-		break;
-	case EPropertyBagPropertyType::String:
-		PinType.PinCategory = UEdGraphSchema_K2::PC_String;
-		break;
-	case EPropertyBagPropertyType::Text:
-		PinType.PinCategory = UEdGraphSchema_K2::PC_Text;
-		break;
-	case EPropertyBagPropertyType::Enum:
-		// @todo: some pin coloring is not correct due to this (byte-as-enum vs enum). 
-		PinType.PinCategory = UEdGraphSchema_K2::PC_Enum;
-		PinType.PinSubCategoryObject = const_cast<UObject*>(Desc.ValueTypeObject.Get());
-		break;
-	case EPropertyBagPropertyType::Struct:
-		PinType.PinCategory = UEdGraphSchema_K2::PC_Struct;
-		PinType.PinSubCategoryObject = const_cast<UObject*>(Desc.ValueTypeObject.Get());
-		break;
-	case EPropertyBagPropertyType::Object:
-		PinType.PinCategory = UEdGraphSchema_K2::PC_Object;
-		PinType.PinSubCategoryObject = const_cast<UObject*>(Desc.ValueTypeObject.Get());
-		break;
-	case EPropertyBagPropertyType::SoftObject:
-		PinType.PinCategory = UEdGraphSchema_K2::PC_SoftObject;
-		PinType.PinSubCategoryObject = const_cast<UObject*>(Desc.ValueTypeObject.Get());
-		break;
-	case EPropertyBagPropertyType::Class:
-		PinType.PinCategory = UEdGraphSchema_K2::PC_Class;
-		PinType.PinSubCategoryObject = const_cast<UObject*>(Desc.ValueTypeObject.Get());
-		break;
-	case EPropertyBagPropertyType::SoftClass:
-		PinType.PinCategory = UEdGraphSchema_K2::PC_SoftClass;
-		PinType.PinSubCategoryObject = const_cast<UObject*>(Desc.ValueTypeObject.Get());
-		break;
-	case EPropertyBagPropertyType::UInt32:	// Warning : Type only partially supported (Blueprint does not support unsigned type)
-		PinType.PinCategory = UEdGraphSchema_K2::PC_Int;
-		break;
-	case EPropertyBagPropertyType::UInt64:	// Warning : Type only partially supported (Blueprint does not support unsigned type)
-		PinType.PinCategory = UEdGraphSchema_K2::PC_Int64;
-		break;
-	default:
-		ensureMsgf(false, TEXT("Unhandled value type %s"), *UEnum::GetValueAsString(Desc.ValueType));
-		break;
+		switch (Type)
+		{
+		case EPropertyBagPropertyType::Bool:
+			OutCategory = UEdGraphSchema_K2::PC_Boolean;
+			break;
+		case EPropertyBagPropertyType::Byte:
+			OutCategory = UEdGraphSchema_K2::PC_Byte;
+			break;
+		case EPropertyBagPropertyType::Int32:
+			OutCategory = UEdGraphSchema_K2::PC_Int;
+			break;
+		case EPropertyBagPropertyType::Int64:
+			OutCategory = UEdGraphSchema_K2::PC_Int64;
+			break;
+		case EPropertyBagPropertyType::Float:
+			OutCategory = UEdGraphSchema_K2::PC_Real;
+			OutSubCategory = UEdGraphSchema_K2::PC_Float;
+			break;
+		case EPropertyBagPropertyType::Double:
+			OutCategory = UEdGraphSchema_K2::PC_Real;
+			OutSubCategory = UEdGraphSchema_K2::PC_Double;
+			break;
+		case EPropertyBagPropertyType::Name:
+			OutCategory = UEdGraphSchema_K2::PC_Name;
+			break;
+		case EPropertyBagPropertyType::String:
+			OutCategory = UEdGraphSchema_K2::PC_String;
+			break;
+		case EPropertyBagPropertyType::Text:
+			OutCategory = UEdGraphSchema_K2::PC_Text;
+			break;
+		case EPropertyBagPropertyType::Enum:
+			// @todo: some pin coloring is not correct due to this (byte-as-enum vs enum). 
+			OutCategory = UEdGraphSchema_K2::PC_Enum;
+			OutSubCategoryObject = const_cast<UObject*>(TypeObject.Get());
+			break;
+		case EPropertyBagPropertyType::Struct:
+			OutCategory = UEdGraphSchema_K2::PC_Struct;
+			OutSubCategoryObject = const_cast<UObject*>(TypeObject.Get());
+			break;
+		case EPropertyBagPropertyType::Object:
+			OutCategory = UEdGraphSchema_K2::PC_Object;
+			OutSubCategoryObject = const_cast<UObject*>(TypeObject.Get());
+			break;
+		case EPropertyBagPropertyType::SoftObject:
+			OutCategory = UEdGraphSchema_K2::PC_SoftObject;
+			OutSubCategoryObject = const_cast<UObject*>(TypeObject.Get());
+			break;
+		case EPropertyBagPropertyType::Class:
+			OutCategory = UEdGraphSchema_K2::PC_Class;
+			OutSubCategoryObject = const_cast<UObject*>(TypeObject.Get());
+			break;
+		case EPropertyBagPropertyType::SoftClass:
+			OutCategory = UEdGraphSchema_K2::PC_SoftClass;
+			OutSubCategoryObject = const_cast<UObject*>(TypeObject.Get());
+			break;
+		case EPropertyBagPropertyType::UInt32:	// Warning : Type only partially supported (Blueprint does not support unsigned type)
+			OutCategory = UEdGraphSchema_K2::PC_Int;
+			break;
+		case EPropertyBagPropertyType::UInt64:	// Warning : Type only partially supported (Blueprint does not support unsigned type)
+			OutCategory = UEdGraphSchema_K2::PC_Int64;
+			break;
+		default:
+			ensureMsgf(false, TEXT("Unhandled value type %s"), *UEnum::GetValueAsString(Type));
+			break;
+		}
+	};
+
+	if (PinType.ContainerType == EPinContainerType::Map)
+	{
+		SetPinTypeFromDescType(Desc.KeyType, Desc.KeyTypeObject, PinType.PinCategory, PinType.PinSubCategory, PinType.PinSubCategoryObject);
+		SetPinTypeFromDescType(Desc.ValueType, Desc.ValueTypeObject, PinType.PinValueType.TerminalCategory, PinType.PinValueType.TerminalSubCategory, PinType.PinValueType.TerminalSubCategoryObject);
+	}
+	else
+	{
+		SetPinTypeFromDescType(Desc.ValueType, Desc.ValueTypeObject, PinType.PinCategory, PinType.PinSubCategory, PinType.PinSubCategoryObject);
 	}
 
 	return PinType;
@@ -2195,7 +2240,7 @@ TSharedPtr<SWidget> FPropertyBagDetails::MakeAddPropertyWidget(TSharedPtr<IPrope
 bool UPropertyBagSchema::SupportsPinTypeContainer(TWeakPtr<const FEdGraphSchemaAction> SchemaAction,
 		const FEdGraphPinType& PinType, const EPinContainerType& ContainerType) const
 {
-	return ContainerType == EPinContainerType::None || ContainerType == EPinContainerType::Array || ContainerType == EPinContainerType::Set;
+	return ContainerType == EPinContainerType::None || ContainerType == EPinContainerType::Array || ContainerType == EPinContainerType::Set || ContainerType == EPinContainerType::Map;
 }
 
 #undef LOCTEXT_NAMESPACE
diff --git a/Engine/Source/Runtime/CoreUObject/Private/StructUtils/PropertyBag.cpp b/Engine/Source/Runtime/CoreUObject/Private/StructUtils/PropertyBag.cpp
--- a/Engine/Source/Runtime/CoreUObject/Private/StructUtils/PropertyBag.cpp
+++ b/Engine/Source/Runtime/CoreUObject/Private/StructUtils/PropertyBag.cpp
@@ -7,6 +7,7 @@
 #include "Misc/NotNull.h"
 #include "Misc/ScopeLock.h"
 #include "Misc/TransactionallySafeCriticalSection.h"
+#include "StructUtils/PropertyBagMapRef.h"
 #include "UObject/EnumProperty.h"
 #include "UObject/Package.h"
 #include "UObject/TextProperty.h"
@@ -44,6 +45,7 @@ namespace FPropertyBagCustomVersion
 		NestedContainerTypes = 2,
 		MetaClass = 3,
 		PropertyFlags = 4,
+		KeyTypes = 5,
 
 		// -----<new versions can be added above this line>-------------------------------------------------
 		VersionPlusOne,
@@ -117,6 +119,7 @@ namespace UE::StructUtils::Private
 		HashBuilder.Update(&Desc.ValueType, sizeof(Desc.ValueType));
 		HashBuilder.Update(&Desc.ContainerTypes, sizeof(Desc.ContainerTypes));
 		HashBuilder.Update(&Desc.PropertyFlags, sizeof(Desc.PropertyFlags));
+		HashBuilder.Update(&Desc.KeyType, sizeof(Desc.KeyType));
 #if WITH_EDITORONLY_DATA
 		for (const FPropertyBagPropertyDescMetaData& MetaData : Desc.MetaData)
 		{
@@ -125,6 +128,7 @@ namespace UE::StructUtils::Private
 		}
 #endif
 		AppendValueObjectHash(HashBuilder, StringBuilder, Desc.ValueTypeObject);
+		AppendValueObjectHash(HashBuilder, StringBuilder, Desc.KeyTypeObject);
 		HashBuilder.Update(StringBuilder.ToString(), StringBuilder.Len() * sizeof(TStringBuilder<256>::ElementType));
 	}
 
@@ -138,6 +142,28 @@ namespace UE::StructUtils::Private
 		return HashBuilder.Finalize().Hash;
 	}
 
+	bool IsTypeHashable(const EPropertyBagPropertyType InType, const UObject* InTypeObject)
+	{
+		switch (InType)
+		{
+		case EPropertyBagPropertyType::Bool:
+		case EPropertyBagPropertyType::Text:
+			return false;
+		case EPropertyBagPropertyType::Struct:
+		{
+			const UScriptStruct* ScriptStruct = Cast<UScriptStruct>(InTypeObject);
+			if (!ScriptStruct)
+			{
+				return false;
+			}
+
+			return ScriptStruct->GetCppStructOps() && ScriptStruct->GetCppStructOps()->HasGetTypeHash();
+		}
+		default:
+			return true;
+		}
+	}
+
 	FPropertyBagContainerTypes GetContainerTypesFromProperty(const FProperty* InSourceProperty)
 	{
 		FPropertyBagContainerTypes ContainerTypes;
@@ -168,6 +194,18 @@ namespace UE::StructUtils::Private
 					ContainerTypes.Reset();
 				}
 			}
+			if (const FMapProperty* MapProperty = CastField<FMapProperty>(InSourceProperty))
+			{
+				if (ContainerTypes.Add(EPropertyBagContainerType::Map))
+				{
+					InSourceProperty = MapProperty->ValueProp;
+					continue;
+				}
+				else // we reached the nested containers limit
+				{
+					ContainerTypes.Reset();
+				}
+			}
 			
 			InSourceProperty = nullptr;
 		}
@@ -259,6 +297,11 @@ namespace UE::StructUtils::Private
 			return GetValueTypeFromProperty(SetProperty->ElementProp);
 		}
 
+		if (const FMapProperty* MapProperty = CastField<FMapProperty>(InSourceProperty))
+		{
+			return GetValueTypeFromProperty(MapProperty->ValueProp);
+		}
+
 		return EPropertyBagPropertyType::None;
 	}
 
@@ -309,6 +352,31 @@ namespace UE::StructUtils::Private
 			return GetValueTypeObjectFromProperty(SetProperty->ElementProp);
 		}
 
+		if (const FMapProperty* MapProperty = CastField<FMapProperty>(InSourceProperty))
+		{
+			return GetValueTypeObjectFromProperty(MapProperty->ValueProp);
+		}
+
+		return nullptr;
+	}
+
+	EPropertyBagPropertyType GetKeyTypeFromProperty(const FProperty* InSourceProperty)
+	{
+		if (const FMapProperty* MapProperty = CastField<FMapProperty>(InSourceProperty))
+		{
+			return GetValueTypeFromProperty(MapProperty->KeyProp);
+		}
+
+		return EPropertyBagPropertyType::None;
+	}
+
+	const UObject* GetKeyTypeObjectFromProperty(const FProperty* InSourceProperty)
+	{
+		if (const FMapProperty* MapProperty = CastField<FMapProperty>(InSourceProperty))
+		{
+			return GetValueTypeObjectFromProperty(MapProperty->KeyProp);
+		}
+
 		return nullptr;
 	}
 
@@ -346,6 +414,22 @@ namespace UE::StructUtils::Private
 						PropertyOwner = SetProperty;
 						break;
 					}
+				case EPropertyBagContainerType::Map:
+					{
+						// create a map property as a container for the tail
+						FMapProperty* MapProperty = new FMapProperty(PropertyOwner, Desc.Name, RF_Public);
+						*ValuePropertyPtr = MapProperty;
+						ValuePropertyPtr = &MapProperty->ValueProp;
+						PropertyOwner = MapProperty;
+
+						// create the key property (container types are not supported for keys)
+						FPropertyBagPropertyDesc KeyDesc;
+						KeyDesc.ValueType = Desc.KeyType;
+						KeyDesc.ValueTypeObject = Desc.KeyTypeObject;
+						KeyDesc.Name = Desc.Name;
+						MapProperty->KeyProp = CreatePropertyFromDesc(KeyDesc, PropertyOwner);
+						break;
+					}
 				default:
 					ensureMsgf(false, TEXT("Unsupported container type %s"), *UEnum::GetValueAsString(BagContainerType));
 					break;
@@ -1682,6 +1766,10 @@ FPropertyBagPropertyDesc::FPropertyBagPropertyDesc(const FName InName, const FPr
 {
 	ValueType = UE::StructUtils::Private::GetValueTypeFromProperty(InSourceProperty);
 	ValueTypeObject = UE::StructUtils::Private::GetValueTypeObjectFromProperty(InSourceProperty);
+
+	KeyType = UE::StructUtils::Private::GetKeyTypeFromProperty(InSourceProperty);
+	KeyTypeObject = UE::StructUtils::Private::GetKeyTypeObjectFromProperty(InSourceProperty);
+
 	// @todo : improve error handling - if we reach the nested containers limit, the Desc will be invalid (empty container types)
 	ContainerTypes = UE::StructUtils::Private::GetContainerTypesFromProperty(InSourceProperty);
 
@@ -1764,6 +1852,12 @@ FArchive& operator<<(FArchive& Ar, FPropertyBagPropertyDesc& Bag)
 	{
 		Ar << Bag.PropertyFlags;
 	}
+
+	if (Ar.CustomVer(FPropertyBagCustomVersion::GUID) >= FPropertyBagCustomVersion::KeyTypes)
+	{
+		Ar << Bag.KeyType;
+		Ar << Bag.KeyTypeObject;
+	}
 	
 	return Ar;
 }
@@ -1869,7 +1963,10 @@ bool FPropertyBagPropertyDesc::CompatibleType(const FPropertyBagPropertyDesc& Ot
 	// Struct and enum must have same value type class
 	if (ValueType == EPropertyBagPropertyType::Enum || ValueType == EPropertyBagPropertyType::Struct)
 	{
-		return ValueTypeObject == Other.ValueTypeObject; 
+		if (ValueTypeObject != Other.ValueTypeObject)
+		{
+			return false;
+		}
 	}
 
 	// Objects should be castable.
@@ -1877,6 +1974,29 @@ bool FPropertyBagPropertyDesc::CompatibleType(const FPropertyBagPropertyDesc& Ot
 	{
 		const UClass* ObjectClass = Cast<const UClass>(ValueTypeObject);
 		const UClass* OtherObjectClass = Cast<const UClass>(Other.ValueTypeObject);
+		if (!UE::StructUtils::Private::CanCastTo(OtherObjectClass, ObjectClass))
+		{
+			return false;
+		}
+	}
+
+	// Keys must match.
+	if (KeyType != Other.KeyType)
+	{
+		return false;
+	}
+
+	// Struct and enum keys must have same key type class
+	if (KeyType == EPropertyBagPropertyType::Enum || KeyType == EPropertyBagPropertyType::Struct)
+	{
+		return KeyTypeObject == Other.KeyTypeObject;
+	}
+
+	// Object keys should be castable.
+	if (KeyType == EPropertyBagPropertyType::Object)
+	{
+		const UClass* ObjectClass = Cast<const UClass>(KeyTypeObject);
+		const UClass* OtherObjectClass = Cast<const UClass>(Other.KeyTypeObject);
 		return UE::StructUtils::Private::CanCastTo(OtherObjectClass, ObjectClass);
 	}
 
@@ -1995,6 +2115,11 @@ EPropertyBagAlterationResult FInstancedPropertyBag::AddProperties(const TConstAr
 			}
 		}
 
+		if (!UE::StructUtils::Private::IsTypeHashable(NewDesc.KeyType, NewDesc.KeyTypeObject))
+		{
+			return EPropertyBagAlterationResult::PropertyKeyTypeNotHashable;
+		}
+
 		FPropertyBagPropertyDesc* ExistingProperty = Descs.FindByPredicate([&NewDesc](const FPropertyBagPropertyDesc& Desc) { return Desc.Name == NewDesc.Name; });
 		if (ExistingProperty != nullptr)
 		{
@@ -2027,14 +2152,19 @@ EPropertyBagAlterationResult FInstancedPropertyBag::AddProperty(const FName InNa
 	return AddProperties({ FPropertyBagPropertyDesc(InName, InValueType, InValueTypeObject) }, bOverwrite);
 }
 
-EPropertyBagAlterationResult FInstancedPropertyBag::AddContainerProperty(const FName InName, const EPropertyBagContainerType InContainerType, const EPropertyBagPropertyType InValueType, const UObject* InValueTypeObject, const bool bOverwrite)
+EPropertyBagAlterationResult FInstancedPropertyBag::AddContainerProperty(const FName InName, const EPropertyBagContainerType InContainerType, const EPropertyBagPropertyType InValueType, const UObject* InValueTypeObject, const bool bOverwrite, const EPropertyBagPropertyType InKeyType, const UObject* InKeyTypeObject)
+{
+	return AddProperties({ FPropertyBagPropertyDesc(InName, InContainerType, InValueType, InValueTypeObject, CPF_Edit, InKeyType, InKeyTypeObject) }, bOverwrite);
+}
+
+EPropertyBagAlterationResult FInstancedPropertyBag::AddContainerProperty(const FName InName, const FPropertyBagContainerTypes InContainerTypes, const EPropertyBagPropertyType InValueType, const UObject* InValueTypeObject, const bool bOverwrite, const EPropertyBagPropertyType InKeyType, const UObject* InKeyTypeObject)
 {
-	return AddProperties({ FPropertyBagPropertyDesc(InName, InContainerType, InValueType, InValueTypeObject) }, bOverwrite);
+	return AddProperties({ FPropertyBagPropertyDesc(InName, InContainerTypes, InValueType, InValueTypeObject, CPF_Edit, InKeyType, InKeyTypeObject) }, bOverwrite);
 }
 
-EPropertyBagAlterationResult FInstancedPropertyBag::AddContainerProperty(const FName InName, const FPropertyBagContainerTypes InContainerTypes, const EPropertyBagPropertyType InValueType, UObject* InValueTypeObject, const bool bOverwrite)
+EPropertyBagAlterationResult FInstancedPropertyBag::AddMapProperty(const FName InName, const EPropertyBagPropertyType InKeyType, const EPropertyBagPropertyType InValueType, const UObject* InKeyTypeObject, const UObject* InValueTypeObject, const EPropertyBagContainerType InNestedContainerType, bool bOverwrite)
 {
-	return AddProperties({ FPropertyBagPropertyDesc(InName, InContainerTypes, InValueType, InValueTypeObject) }, bOverwrite);
+	return AddProperties({ FPropertyBagPropertyDesc(InName, { EPropertyBagContainerType::Map, InNestedContainerType }, InValueType, InValueTypeObject, CPF_Edit, InKeyType, InKeyTypeObject)}, bOverwrite);
 }
 
 EPropertyBagAlterationResult FInstancedPropertyBag::AddProperty(const FName InName, const FProperty* InSourceProperty, const bool bOverwrite)
@@ -2856,6 +2986,50 @@ TValueOrError<const FPropertyBagSetRef, EPropertyBagResult> FInstancedPropertyBa
 	return MakeValue(FPropertyBagSetRef(Desc, Address));
 }
 
+TValueOrError<FPropertyBagMapRef, EPropertyBagResult> FInstancedPropertyBag::GetMutableMapRef(const FPropertyBagPropertyDesc& Desc)
+{
+	if (!OwnsPropertyDesc(Desc))
+	{
+		return MakeError(EPropertyBagResult::PropertyNotFound);
+	}
+	check(Desc.CachedProperty);
+
+	if (Desc.ContainerTypes.GetFirstContainerType() != EPropertyBagContainerType::Map)
+	{
+		return MakeError(EPropertyBagResult::TypeMismatch);
+	}
+
+	const void* Address = GetValueAddress(&Desc);
+	if (Address == nullptr)
+	{
+		return MakeError(EPropertyBagResult::PropertyNotFound);
+	}
+
+	return MakeValue(FPropertyBagMapRef(Desc, Address));
+}
+
+TValueOrError<const FPropertyBagMapRef, EPropertyBagResult> FInstancedPropertyBag::GetMapRef(const FPropertyBagPropertyDesc& Desc) const
+{
+	if (!OwnsPropertyDesc(Desc))
+	{
+		return MakeError(EPropertyBagResult::PropertyNotFound);
+	}
+	check(Desc.CachedProperty);
+
+	if (Desc.ContainerTypes.GetFirstContainerType() != EPropertyBagContainerType::Map)
+	{
+		return MakeError(EPropertyBagResult::TypeMismatch);
+	}
+
+	const void* Address = GetValueAddress(&Desc);
+	if (Address == nullptr)
+	{
+		return MakeError(EPropertyBagResult::PropertyNotFound);
+	}
+
+	return MakeValue(FPropertyBagMapRef(Desc, Address));
+}
+
 EPropertyBagResult FInstancedPropertyBag::SetValueBool(const FName Name, const bool bInValue)
 {
 	const FPropertyBagPropertyDesc* Desc = FindPropertyDescByName(Name);
@@ -3248,6 +3422,52 @@ TValueOrError<const FPropertyBagSetRef, EPropertyBagResult> FInstancedPropertyBa
 	return MakeValue(FPropertyBagSetRef(*Desc, Address));
 }
 
+TValueOrError<FPropertyBagMapRef, EPropertyBagResult> FInstancedPropertyBag::GetMutableMapRef(const FName Name)
+{
+	const FPropertyBagPropertyDesc* Desc = FindPropertyDescByName(Name);
+	if (Desc == nullptr)
+	{
+		return MakeError(EPropertyBagResult::PropertyNotFound);
+	}
+	check(Desc->CachedProperty);
+
+	if (Desc->ContainerTypes.GetFirstContainerType() != EPropertyBagContainerType::Map)
+	{
+		return MakeError(EPropertyBagResult::TypeMismatch);
+	}
+
+	const void* Address = GetValueAddress(Desc);
+	if (Address == nullptr)
+	{
+		return MakeError(EPropertyBagResult::PropertyNotFound);
+	}
+
+	return MakeValue(FPropertyBagMapRef(*Desc, Address));
+}
+
+TValueOrError<const FPropertyBagMapRef, EPropertyBagResult> FInstancedPropertyBag::GetMapRef(const FName Name) const
+{
+	const FPropertyBagPropertyDesc* Desc = FindPropertyDescByName(Name);
+	if (Desc == nullptr)
+	{
+		return MakeError(EPropertyBagResult::PropertyNotFound);
+	}
+	check(Desc->CachedProperty);
+
+	if (Desc->ContainerTypes.GetFirstContainerType() != EPropertyBagContainerType::Map)
+	{
+		return MakeError(EPropertyBagResult::TypeMismatch);
+	}
+
+	const void* Address = GetValueAddress(Desc);
+	if (Address == nullptr)
+	{
+		return MakeError(EPropertyBagResult::PropertyNotFound);
+	}
+
+	return MakeValue(FPropertyBagMapRef(*Desc, Address));
+}
+
 bool FInstancedPropertyBag::Identical(const FInstancedPropertyBag* Other, const uint32 PortFlags) const
 {
 	return Other && Value.Identical(&Other->Value, PortFlags);
@@ -3292,6 +3512,10 @@ bool FInstancedPropertyBag::Serialize(FArchive& Ar)
 				{
 					Ar.Preload(const_cast<UObject*>(PropDesc.ValueTypeObject.Get()));
 				}
+				if (PropDesc.KeyTypeObject)
+				{
+					Ar.Preload(const_cast<UObject*>(PropDesc.KeyTypeObject.Get()));
+				}
 			}
 			
 			BagStruct = const_cast<UPropertyBag*>(UPropertyBag::GetOrCreateFromDescs(PropertyDescs));
@@ -3335,6 +3559,16 @@ bool FInstancedPropertyBag::Serialize(FArchive& Ar)
 						Desc.ValueTypeObject = UserDefinedStruct->PrimaryStruct.Get();
 					}
 				}
+				if (Desc.KeyType == EPropertyBagPropertyType::Struct)
+				{
+					const UUserDefinedStruct* UserDefinedStruct = Cast<UUserDefinedStruct>(Desc.KeyTypeObject);
+					if (UserDefinedStruct
+						&& UserDefinedStruct->Status == EUserDefinedStructureStatus::UDSS_Duplicate
+						&& UserDefinedStruct->PrimaryStruct.IsValid())
+					{
+						Desc.KeyTypeObject = UserDefinedStruct->PrimaryStruct.Get();
+					}
+				}
 			}
 #endif // WITH_EDITOR
 
@@ -3398,6 +3632,16 @@ void FInstancedPropertyBag::AddStructReferencedObjects(FReferenceCollector& Coll
 							}
 						}					
 					}
+					if (Desc.KeyType == EPropertyBagPropertyType::Struct)
+					{
+						if (const UUserDefinedStruct* UserDefinedStruct = Cast<UUserDefinedStruct>(Desc.KeyTypeObject))
+						{
+							if (UserDefinedStruct == StructureToReinstance->PrimaryStruct)
+							{
+								Desc.KeyTypeObject = StructureToReinstance;
+							}
+						}
+					}
 				}
 
 				const UPropertyBag* NewBag = UPropertyBag::GetOrCreateFromDescs(PropertyDescs);
@@ -3446,6 +3690,10 @@ void FInstancedPropertyBag::GetPreloadDependencies(TArray<UObject*>& OutDeps)
 			{
 				OutDeps.Add(const_cast<UObject*>(Desc.ValueTypeObject.Get()));
 			}
+			if (Desc.KeyTypeObject)
+			{
+				OutDeps.Add(const_cast<UObject*>(Desc.KeyTypeObject.Get()));
+			}
 		}
 
 		// Report indirect dependencies of the instanced property bag struct
@@ -3979,41 +4227,51 @@ TArray<FPropertyBagPropertyDesc> UPropertyBag::SanitizePropertyDescs(const TArra
 {
 	TArray<FPropertyBagPropertyDesc> Result = TArray<FPropertyBagPropertyDesc>(PropertyDescs.GetData(), PropertyDescs.Num());
 
-	for (FPropertyBagPropertyDesc& Desc : Result)
+	// Fix missing structs, enums, and objects.
+	auto CheckType = [](const FName PropertyName, EPropertyBagPropertyType InType, TObjectPtr<const UObject>& InOutTypeObject)
 	{
-		// Fix missing structs, enums, and objects.
-		if (Desc.ValueType == EPropertyBagPropertyType::Struct)
+		if (InType == EPropertyBagPropertyType::Struct)
 		{
-			if (Desc.ValueTypeObject == nullptr || Desc.ValueTypeObject->GetClass()->IsChildOf(UScriptStruct::StaticClass()) == false)
+			if (InOutTypeObject == nullptr || InOutTypeObject->GetClass()->IsChildOf(UScriptStruct::StaticClass()) == false)
 			{
-				UE_LOG(LogCore, Warning, TEXT("PropertyBag: Struct property '%s' is missing type."), *Desc.Name.ToString());
-				Desc.ValueTypeObject = FPropertyBagMissingStruct::StaticStruct();
+				UE_LOG(LogCore, Warning, TEXT("PropertyBag: Struct property '%s' is missing type."), *PropertyName.ToString());
+				InOutTypeObject = FPropertyBagMissingStruct::StaticStruct();
 			}
 		}
-		else if (Desc.ValueType == EPropertyBagPropertyType::Enum)
+		else if (InType == EPropertyBagPropertyType::Enum)
 		{
-			if (Desc.ValueTypeObject == nullptr || Desc.ValueTypeObject->GetClass()->IsChildOf(UEnum::StaticClass()) == false)
+			if (InOutTypeObject == nullptr || InOutTypeObject->GetClass()->IsChildOf(UEnum::StaticClass()) == false)
 			{
-				UE_LOG(LogCore, Warning, TEXT("PropertyBag: Enum property '%s' is missing type."), *Desc.Name.ToString());
-				Desc.ValueTypeObject = StaticEnum<EPropertyBagMissingEnum>();
+				UE_LOG(LogCore, Warning, TEXT("PropertyBag: Enum property '%s' is missing type."), *PropertyName.ToString());
+				InOutTypeObject = StaticEnum<EPropertyBagMissingEnum>();
 			}
 		}
-		else if (Desc.ValueType == EPropertyBagPropertyType::Object || Desc.ValueType == EPropertyBagPropertyType::SoftObject)
+		else if (InType == EPropertyBagPropertyType::Object || InType == EPropertyBagPropertyType::SoftObject)
 		{
-			if (Desc.ValueTypeObject == nullptr)
+			if (InOutTypeObject == nullptr)
 			{
-				UE_LOG(LogCore, Warning, TEXT("PropertyBag: Object property '%s' is missing type."), *Desc.Name.ToString());
-				Desc.ValueTypeObject = UPropertyBagMissingObject::StaticClass();
+				UE_LOG(LogCore, Warning, TEXT("PropertyBag: Object property '%s' is missing type."), *PropertyName.ToString());
+				InOutTypeObject = UPropertyBagMissingObject::StaticClass();
 			}
 		}
-		else if (Desc.ValueType == EPropertyBagPropertyType::Class || Desc.ValueType == EPropertyBagPropertyType::SoftClass)
+		else if (InType == EPropertyBagPropertyType::Class || InType == EPropertyBagPropertyType::SoftClass)
 		{
-			if (Desc.ValueTypeObject == nullptr || Desc.ValueTypeObject->GetClass()->IsChildOf(UClass::StaticClass()) == false)
+			if (InOutTypeObject == nullptr || InOutTypeObject->GetClass()->IsChildOf(UClass::StaticClass()) == false)
 			{
-				UE_LOG(LogCore, Warning, TEXT("PropertyBag: Class property '%s' is missing type."), *Desc.Name.ToString());
-				Desc.ValueTypeObject = UPropertyBagMissingObject::StaticClass();
+				UE_LOG(LogCore, Warning, TEXT("PropertyBag: Class property '%s' is missing type."), *PropertyName.ToString());
+				InOutTypeObject = UPropertyBagMissingObject::StaticClass();
 			}
 		}
+	};
+
+	for (FPropertyBagPropertyDesc& Desc : Result)
+	{
+		CheckType(Desc.Name, Desc.ValueType, Desc.ValueTypeObject);
+
+		if (Desc.KeyType != EPropertyBagPropertyType::None)
+		{
+			CheckType(Desc.Name, Desc.KeyType, Desc.KeyTypeObject);
+		}
 
 		// Fix ID
 		if (!Desc.ID.IsValid())
@@ -4118,6 +4376,21 @@ bool UPropertyBag::ContainsUserDefinedStruct(const UUserDefinedStruct* UserDefin
 				}
 			}
 		}
+		if (Desc.KeyType == EPropertyBagPropertyType::Struct)
+		{
+			if (const UUserDefinedStruct* OwnedUserDefinedStruct = Cast<UUserDefinedStruct>(Desc.KeyTypeObject))
+			{
+				if (OwnedUserDefinedStruct == UserDefinedStruct
+					|| OwnedUserDefinedStruct->PrimaryStruct == UserDefinedStruct
+					|| OwnedUserDefinedStruct == UserDefinedStruct->PrimaryStruct
+					|| (OwnedUserDefinedStruct->PrimaryStruct == UserDefinedStruct->PrimaryStruct
+						&& OwnedUserDefinedStruct->PrimaryStruct != nullptr)
+					)
+				{
+					return true;
+				}
+			}
+		}
 	}
 	return false;
 }
diff --git a/Engine/Source/Runtime/CoreUObject/Public/StructUtils/PropertyBag.h b/Engine/Source/Runtime/CoreUObject/Public/StructUtils/PropertyBag.h
--- a/Engine/Source/Runtime/CoreUObject/Public/StructUtils/PropertyBag.h
+++ b/Engine/Source/Runtime/CoreUObject/Public/StructUtils/PropertyBag.h
@@ -46,6 +46,7 @@ enum class EPropertyBagContainerType : uint8
 	None,
 	Array,
 	Set,
+	Map,
 
 	Count UMETA(Hidden)
 };
@@ -188,7 +189,8 @@ enum class EPropertyBagAlterationResult : uint8
 	PropertyNameInvalidCharacters,       // The property name contains illegal characters. Consider using FInstancedPropertyBag::SanitizePropertyName.
 	SourcePropertyNotFound,              // The source property or property name was not found.
 	TargetPropertyNotFound,              // The target property or property name was not found.
-	TargetPropertyAlreadyExists          // The target property or property name already exists.
+	TargetPropertyAlreadyExists,         // The target property or property name already exists.
+	PropertyKeyTypeNotHashable,          // The property has Map container type and the given key type is not hashable.
 };
 
 USTRUCT()
@@ -259,21 +261,25 @@ struct FPropertyBagPropertyDesc
 		, ValueType(InValueType)
 	{
 	}
-	FPropertyBagPropertyDesc(const FName InName, const EPropertyBagContainerType InContainerType, const EPropertyBagPropertyType InValueType, const UObject* InValueTypeObject = nullptr, EPropertyFlags InPropertyFlags = CPF_Edit)
+	FPropertyBagPropertyDesc(const FName InName, const EPropertyBagContainerType InContainerType, const EPropertyBagPropertyType InValueType, const UObject* InValueTypeObject = nullptr, EPropertyFlags InPropertyFlags = CPF_Edit, const EPropertyBagPropertyType InKeyType = EPropertyBagPropertyType::None, const UObject* InKeyTypeObject = nullptr)
 		: ValueTypeObject(InValueTypeObject)
 		, Name(InName)
 		, ValueType(InValueType)
 		, ContainerTypes(InContainerType)
 		, PropertyFlags((uint64)InPropertyFlags)
+		, KeyType(InKeyType)
+		, KeyTypeObject(InKeyTypeObject)
 	{
 	}
 
-	FPropertyBagPropertyDesc(const FName InName, const FPropertyBagContainerTypes& InNestedContainers, const EPropertyBagPropertyType InValueType, UObject* InValueTypeObject = nullptr, EPropertyFlags InPropertyFlags = CPF_Edit)
+	FPropertyBagPropertyDesc(const FName InName, const FPropertyBagContainerTypes& InNestedContainers, const EPropertyBagPropertyType InValueType, const UObject* InValueTypeObject = nullptr, EPropertyFlags InPropertyFlags = CPF_Edit, const EPropertyBagPropertyType InKeyType = EPropertyBagPropertyType::None, const UObject* InKeyTypeObject = nullptr)
 		: ValueTypeObject(InValueTypeObject)
 		, Name(InName)
 		, ValueType(InValueType)
 		, ContainerTypes(InNestedContainers)
 		, PropertyFlags((uint64)InPropertyFlags)
+		, KeyType(InKeyType)
+		, KeyTypeObject(InKeyTypeObject)
 	{
 	}
 
@@ -344,6 +350,14 @@ struct FPropertyBagPropertyDesc
 	UPROPERTY(EditAnywhere, Category = "Default")
 	uint64 PropertyFlags = (uint64)CPF_Edit;
 
+	/** Type of this property's key (for map containers only) */
+	UPROPERTY(EditAnywhere, Category = "Default")
+	EPropertyBagPropertyType KeyType = EPropertyBagPropertyType::None;
+
+	/** Pointer to object that defines the Enum, Struct, or Class of this property's key (for map containers only) */
+	UPROPERTY(EditAnywhere, Category = "Default")
+	TObjectPtr<const UObject> KeyTypeObject = nullptr;
+
 #if WITH_EDITORONLY_DATA
 	/** Editor-only metadata for CachedProperty */
 	UPROPERTY(EditAnywhere, Category="Default")
@@ -409,6 +423,7 @@ private:
 class UPropertyBag;
 class FPropertyBagArrayRef;
 class FPropertyBagSetRef;
+class FPropertyBagMapRef;
 
 USTRUCT()
 struct FInstancedPropertyBag
@@ -491,9 +506,11 @@ struct FInstancedPropertyBag
 	 * @param InValueType Type of the new property
 	 * @param InValueTypeObject Type object (for struct, class, enum) of the new property
 	 * @param bOverwrite Overwrite the property if it already exists.
+	 * @param InKeyType Type of the new property's key
+	 * @param InKeyTypeObject Type object (for struct, class, enum) of the new property's key
 	 * @return The result of the alteration.
 	 */
-	UE_API EPropertyBagAlterationResult AddContainerProperty(const FName InName, const EPropertyBagContainerType InContainerType, const EPropertyBagPropertyType InValueType, const UObject* InValueTypeObject = nullptr, bool bOverwrite = true);
+	UE_API EPropertyBagAlterationResult AddContainerProperty(const FName InName, const EPropertyBagContainerType InContainerType, const EPropertyBagPropertyType InValueType, const UObject* InValueTypeObject = nullptr, bool bOverwrite = true, const EPropertyBagPropertyType InKeyType = EPropertyBagPropertyType::None, const UObject* InKeyTypeObject = nullptr);
 
 	/**
 	 * Adds a new container property to the bag. If property of same name already exists, it will be replaced with the new type.
@@ -502,9 +519,24 @@ struct FInstancedPropertyBag
 	 * @param InValueType Type of the new property
 	 * @param InValueTypeObject Type object (for struct, class, enum) of the new property
 	 * @param bOverwrite Overwrite the property if it already exists.
+	 * @param InKeyType Type of the new property's key
+	 * @param InKeyTypeObject Type object (for struct, class, enum) of the new property's key
+	 * @return The result of the alteration.
+	 */
+	UE_API EPropertyBagAlterationResult AddContainerProperty(const FName InName, const FPropertyBagContainerTypes InContainerTypes, const EPropertyBagPropertyType InValueType, const UObject* InValueTypeObject, bool bOverwrite = true, const EPropertyBagPropertyType InKeyType = EPropertyBagPropertyType::None, const UObject* InKeyTypeObject = nullptr);
+
+	/**
+	 * Adds a new map property to the bag. If property of same name already exists, it will be replaced with the new type.
+	 * @param InName Name of the new property
+	 * @param InKeyType Type of the new property's key
+	 * @param InValueType Type of the new property
+	 * @param InKeyTypeObject Type object (for struct, class, enum) of the new property's key
+	 * @param InValueTypeObject Type object (for struct, class, enum) of the new property
+	 * @param InNestedContainerType Type of (optional) nested container to create
+	 * @param bOverwrite Overwrite the property if it already exists.
 	 * @return The result of the alteration.
 	 */
-	UE_API EPropertyBagAlterationResult AddContainerProperty(const FName InName, const FPropertyBagContainerTypes InContainerTypes, const EPropertyBagPropertyType InValueType, UObject* InValueTypeObject, bool bOverwrite = true);
+	UE_API EPropertyBagAlterationResult AddMapProperty(const FName InName, const EPropertyBagPropertyType InKeyType, const EPropertyBagPropertyType InValueType, const UObject* InKeyTypeObject = nullptr, const UObject* InValueTypeObject = nullptr, const EPropertyBagContainerType InNestedContainerType = EPropertyBagContainerType::None, bool bOverwrite = true);
 
 	/**
 	 * Adds a new property to the bag. Property type duplicated from source property to. If property of same name already exists, it will be replaced with the new type.
@@ -716,6 +748,8 @@ struct FInstancedPropertyBag
 	UE_API TValueOrError<const FPropertyBagArrayRef, EPropertyBagResult> GetArrayRef(const FPropertyBagPropertyDesc& Desc) const;
 	UE_API TValueOrError<FPropertyBagSetRef, EPropertyBagResult> GetMutableSetRef(const FPropertyBagPropertyDesc& Desc);
 	UE_API TValueOrError<const FPropertyBagSetRef, EPropertyBagResult> GetSetRef(const FPropertyBagPropertyDesc& Desc) const;
+	UE_API TValueOrError<FPropertyBagMapRef, EPropertyBagResult> GetMutableMapRef(const FPropertyBagPropertyDesc& Desc);
+	UE_API TValueOrError<const FPropertyBagMapRef, EPropertyBagResult> GetMapRef(const FPropertyBagPropertyDesc& Desc) const;
 
 	/** @return enum value of specified type. */
 	template <typename T>
@@ -895,6 +929,20 @@ struct FInstancedPropertyBag
 	*/
 	UE_API TValueOrError<const FPropertyBagSetRef, EPropertyBagResult> GetSetRef(const FName Name) const;
 
+	/**
+	 * Returns helper class to modify and access a map property.
+	 * Note: The map reference is not valid after the layout of the referenced property bag has changed!
+	 * @returns helper class to modify and access maps
+	*/
+	UE_API TValueOrError<FPropertyBagMapRef, EPropertyBagResult> GetMutableMapRef(const FName Name);
+
+	/**
+	 * Returns helper class to access a map property.
+	 * Note: The map reference is not valid after the layout of the referenced property bag has changed!
+	 * @returns helper class to modify and access maps
+	*/
+	UE_API TValueOrError<const FPropertyBagMapRef, EPropertyBagResult> GetMapRef(const FName Name) const;
+
 	UE_API bool Identical(const FInstancedPropertyBag* Other, uint32 PortFlags) const;
 	UE_API bool Serialize(FArchive& Ar);
 	UE_API void AddStructReferencedObjects(FReferenceCollector& Collector);
diff --git a/Engine/Source/Runtime/CoreUObject/Public/StructUtils/PropertyBagMapRef.h b/Engine/Source/Runtime/CoreUObject/Public/StructUtils/PropertyBagMapRef.h
new file mode 100644
--- /dev/null
+++ b/Engine/Source/Runtime/CoreUObject/Public/StructUtils/PropertyBagMapRef.h
@@ -0,0 +1,464 @@
+// Copyright Epic Games, Inc. All Rights Reserved.
+
+#pragma once
+
+#include "Concepts/BaseStructureProvider.h"
+#include "StructUtils/PropertyBag.h"
+
+class FPropertyBagMapRef : private FScriptMapHelper
+{
+public:
+	FORCEINLINE FPropertyBagMapRef(const FPropertyBagPropertyDesc& InDesc, const void* InMap)
+		: FScriptMapHelper(CastField<FMapProperty>(InDesc.CachedProperty), InMap)
+	{
+		const FMapProperty* MapProperty = CastField<FMapProperty>(InDesc.CachedProperty);
+		check(MapProperty);
+		check(MapProperty->KeyProp);
+		check(MapProperty->ValueProp);
+		// Create dummy desc for the inner properties ( key + value ). 
+		KeyDesc.ValueType = InDesc.KeyType;
+		KeyDesc.ValueTypeObject = InDesc.KeyTypeObject;
+		KeyDesc.CachedProperty = MapProperty->KeyProp;
+		KeyDesc.ContainerTypes.Reset(); // Key's can't be containers		
+		ValueDesc.ValueType = InDesc.ValueType;
+		ValueDesc.ValueTypeObject = InDesc.ValueTypeObject;
+		ValueDesc.CachedProperty = MapProperty->ValueProp;
+		ValueDesc.ContainerTypes = InDesc.ContainerTypes;
+		ValueDesc.ContainerTypes.PopHead();
+		if (ValueDesc.ContainerTypes.GetFirstContainerType() == EPropertyBagContainerType::Map)
+		{
+			ValueDesc.KeyType = InDesc.KeyType;
+			ValueDesc.KeyTypeObject = InDesc.KeyTypeObject;
+		}
+	}
+
+	using FScriptMapHelper::Num;
+
+	void Empty()
+	{
+		FScriptMapHelper::EmptyValues();
+	}
+
+	template <typename KeyType>
+	TValueOrError<bool, EPropertyBagResult> Contains(const KeyType& InKey) const
+	{
+		EPropertyBagResult Result = CheckType<KeyType>(KeyDesc, InKey);
+		if (Result != EPropertyBagResult::Success)
+		{
+			return MakeError(Result);
+		}
+
+		const int32 PairIndex = Call<KeyType, int32>([this](const void* Key)
+		{
+			return const_cast<FPropertyBagMapRef*>(this)->FindMapPairIndexFromHash(Key);
+		}, KeyDesc, InKey);
+
+		return MakeValue(PairIndex != INDEX_NONE);
+	}
+
+	template <typename KeyType, typename ValueType>
+	TValueOrError<ValueType, EPropertyBagResult> Find(const KeyType& InKey) const
+	{
+		EPropertyBagResult Result = CheckType<KeyType>(KeyDesc, InKey);
+		if (Result != EPropertyBagResult::Success)
+		{
+			return MakeError(Result);
+		}
+
+		uint8* RawValue = Call<KeyType, uint8*>([this](const void* Key)
+		{
+			return const_cast<FPropertyBagMapRef*>(this)->FindValueFromHash(Key);
+		}, KeyDesc, InKey);
+		
+		if (RawValue == nullptr)
+		{
+			return MakeError(EPropertyBagResult::PropertyNotFound);
+		}
+
+		ValueType Value;
+		Convert(ValueDesc, RawValue, Value);
+
+		Result = CheckType<ValueType>(ValueDesc, Value);
+		if (Result != EPropertyBagResult::Success)
+		{
+			return MakeError(Result);
+		}
+
+		return MakeValue(Value);
+	}
+
+	template <typename KeyType>
+	EPropertyBagResult Add(const KeyType& InKey)
+	{
+		TValueOrError<bool, EPropertyBagResult> Result = Contains(InKey);
+		if (Result.HasError())
+		{
+			return Result.GetError();
+		}
+		if (Result.GetValue())
+		{
+			return EPropertyBagResult::DuplicatedValue;
+		}
+
+		void* RawValue = Call<KeyType, void*>([this](const void* Key) { return FScriptMapHelper::FindOrAdd(Key); }, KeyDesc, InKey);
+		check(RawValue);
+
+		return EPropertyBagResult::Success;
+	}
+
+	template <typename KeyType, typename ValueType>
+	EPropertyBagResult Add(const KeyType& InKey, const ValueType& InValue)
+	{
+		EPropertyBagResult Result = CheckType<KeyType>(KeyDesc, InKey);
+		if (Result != EPropertyBagResult::Success)
+		{
+			return Result;
+		}
+
+		Result = CheckType<ValueType>(ValueDesc, InValue);
+		if (Result != EPropertyBagResult::Success)
+		{
+			return Result;
+		}
+
+		return AddPair(KeyDesc, InKey, ValueDesc, InValue);
+	}
+
+	template <typename KeyType>
+	EPropertyBagResult Remove(const KeyType& InKey)
+	{
+		EPropertyBagResult Result = CheckType<KeyType>(KeyDesc, InKey);
+		if (Result != EPropertyBagResult::Success)
+		{
+			return Result;
+		}
+
+		const bool bWasRemoved = Call<KeyType, bool>([this](const void* Key) { return FScriptMapHelper::RemovePair(Key); }, KeyDesc, InKey);
+
+		return bWasRemoved ? EPropertyBagResult::Success : EPropertyBagResult::PropertyNotFound;
+	}
+
+	template <typename KeyType>
+	TValueOrError<FPropertyBagArrayRef, EPropertyBagResult> GetMutableNestedArrayRef(const KeyType& InKey)
+	{
+		EPropertyBagResult Result = CheckType<KeyType>(KeyDesc, InKey);
+		if (Result != EPropertyBagResult::Success)
+		{
+			return MakeError(Result);
+		}
+
+		if (ValueDesc.ContainerTypes.GetFirstContainerType() != EPropertyBagContainerType::Array)
+		{
+			return MakeError(EPropertyBagResult::TypeMismatch);
+		}
+
+		uint8* RawValue = Call<KeyType, uint8*>([this](const void* Key) { return FScriptMapHelper::FindValueFromHash(Key); }, KeyDesc, InKey);
+		if (RawValue == nullptr)
+		{
+			return MakeError(EPropertyBagResult::PropertyNotFound);
+		}
+
+		return MakeValue(FPropertyBagArrayRef(ValueDesc, RawValue));
+	}
+
+	template <typename KeyType>
+	TValueOrError<const FPropertyBagArrayRef, EPropertyBagResult> GetNestedArrayRef(const KeyType& InKey)
+	{
+		TValueOrError<FPropertyBagArrayRef, EPropertyBagResult> Result = GetMutableNestedArrayRef(InKey);
+		if (Result.HasError())
+		{
+			return MakeError(Result.GetError());
+		}
+
+		return MakeValue(Result.GetValue());
+	}
+
+	template <typename KeyType>
+	TValueOrError<FPropertyBagSetRef, EPropertyBagResult> GetMutableNestedSetRef(const KeyType& InKey)
+	{
+		EPropertyBagResult Result = CheckType<KeyType>(KeyDesc, InKey);
+		if (Result != EPropertyBagResult::Success)
+		{
+			return MakeError(Result);
+		}
+
+		if (ValueDesc.ContainerTypes.GetFirstContainerType() != EPropertyBagContainerType::Set)
+		{
+			return MakeError(EPropertyBagResult::TypeMismatch);
+		}
+
+		uint8* RawValue = Call<KeyType, uint8*>([this](const void* Key) { return FScriptMapHelper::FindValueFromHash(Key); }, KeyDesc, InKey);
+		if (RawValue == nullptr)
+		{
+			return MakeError(EPropertyBagResult::PropertyNotFound);
+		}
+
+		return MakeValue(FPropertyBagSetRef(ValueDesc, RawValue));
+	}
+
+	template <typename KeyType>
+	TValueOrError<const FPropertyBagSetRef, EPropertyBagResult> GetNestedSetRef(const KeyType& InKey)
+	{
+		TValueOrError<FPropertyBagSetRef, EPropertyBagResult> Result = GetMutableNestedSetRef(InKey);
+		if (Result.HasError())
+		{
+			return MakeError(Result.GetError());
+		}
+
+		return MakeValue(Result.GetValue());
+	}
+
+	template <typename KeyType>
+	TValueOrError<FPropertyBagMapRef, EPropertyBagResult> GetMutableNestedMapRef(const KeyType& InKey)
+	{
+		EPropertyBagResult Result = CheckType<KeyType>(KeyDesc, InKey);
+		if (Result != EPropertyBagResult::Success)
+		{
+			return MakeError(Result);
+		}
+
+		if (ValueDesc.ContainerTypes.GetFirstContainerType() != EPropertyBagContainerType::Map)
+		{
+			return MakeError(EPropertyBagResult::TypeMismatch);
+		}
+
+		uint8* RawValue = Call<KeyType, uint8*>([this](const void* Key) { return FScriptMapHelper::FindValueFromHash(Key); }, KeyDesc, InKey);
+		if (RawValue == nullptr)
+		{
+			return MakeError(EPropertyBagResult::PropertyNotFound);
+		}
+
+		return MakeValue(FPropertyBagMapRef(ValueDesc, RawValue));
+	}
+
+	template <typename KeyType>
+	TValueOrError<const FPropertyBagMapRef, EPropertyBagResult> GetNestedMapRef(const KeyType& InKey)
+	{
+		TValueOrError<FPropertyBagMapRef, EPropertyBagResult> Result = GetMutableNestedMapRef(InKey);
+		if (Result.HasError())
+		{
+			return MakeError(Result.GetError());
+		}
+
+		return MakeValue(Result.GetValue());
+	}
+
+private:
+
+	template <typename T, typename TEnableIf<TIsEnumClass<T>::Value>::Type* = nullptr>
+	static EPropertyBagResult CheckType(FPropertyBagPropertyDesc InDesc, const T InValue)
+	{
+		if (!InDesc.ContainerTypes.IsEmpty())
+		{
+			return EPropertyBagResult::TypeMismatch;
+		}
+
+		if (InDesc.ValueType == EPropertyBagPropertyType::Enum)
+		{
+			const FEnumProperty* EnumProperty = CastField<FEnumProperty>(InDesc.CachedProperty);
+
+			if (EnumProperty && EnumProperty->GetEnum() == StaticEnum<T>())
+			{
+				return EPropertyBagResult::Success;
+			}
+		}
+
+		return EPropertyBagResult::TypeMismatch;
+	}
+
+	template <typename T, typename TEnableIf<TModels_V<CBaseStructureProvider, T>>::Type* = nullptr>
+	static EPropertyBagResult CheckType(FPropertyBagPropertyDesc InDesc, const T InValue)
+	{
+		if (!InDesc.ContainerTypes.IsEmpty())
+		{
+			return EPropertyBagResult::TypeMismatch;
+		}
+
+		if (InDesc.ValueType == EPropertyBagPropertyType::Struct)
+		{
+			const FStructProperty* StructProperty = CastField<FStructProperty>(InDesc.CachedProperty);
+
+			if (StructProperty && StructProperty->Struct == TBaseStructure<T>::Get())
+			{
+				return EPropertyBagResult::Success;
+			}
+		}
+		else if (TIsDerivedFrom<T, FSoftObjectPath>::Value && (InDesc.ValueType == EPropertyBagPropertyType::SoftObject || InDesc.ValueType == EPropertyBagPropertyType::SoftClass))
+		{
+			return EPropertyBagResult::Success;
+		}
+
+		return EPropertyBagResult::TypeMismatch;
+	}
+
+	template <typename T, typename TEnableIf<TIsDerivedFrom<std::remove_pointer_t<T>, UObject>::Value>::Type* = nullptr>
+	static EPropertyBagResult CheckType(FPropertyBagPropertyDesc InDesc, const UObject* InValue)
+	{
+		if (!InDesc.ContainerTypes.IsEmpty())
+		{
+			return EPropertyBagResult::TypeMismatch;
+		}
+
+		if (InDesc.ValueType == EPropertyBagPropertyType::Object || InDesc.ValueType == EPropertyBagPropertyType::SoftObject)
+		{
+			if (!IsValid(InValue))
+			{
+				return EPropertyBagResult::Success;
+			}
+			else if (const FObjectPropertyBase* ObjectProperty = CastField<FObjectPropertyBase>(InDesc.CachedProperty))
+			{
+				const UClass* Class = InValue->GetClass();
+				const UClass* PropertyClass = ObjectProperty->PropertyClass;
+
+				if (Class && PropertyClass && Class->IsChildOf(PropertyClass))
+				{
+					return EPropertyBagResult::Success;
+				}
+			}
+		}
+		else if (InDesc.ValueType == EPropertyBagPropertyType::Class || InDesc.ValueType == EPropertyBagPropertyType::SoftClass)
+		{
+			if (!IsValid(InValue))
+			{
+				return EPropertyBagResult::Success;
+			}
+
+			const UClass* Class = Cast<UClass>(InValue);
+			const UClass* PropertyClass = nullptr;
+
+			if (const FClassProperty* ClassProperty = CastField<FClassProperty>(InDesc.CachedProperty))
+			{
+				PropertyClass = ClassProperty->MetaClass;
+			}
+			else if (const FSoftClassProperty* SoftClassProperty = CastField<FSoftClassProperty>(InDesc.CachedProperty))
+			{
+				PropertyClass = SoftClassProperty->MetaClass;
+			}
+
+			if (Class && PropertyClass && Class->IsChildOf(PropertyClass))
+			{
+				return EPropertyBagResult::Success;
+			}
+		}
+
+		return EPropertyBagResult::TypeMismatch;
+	}
+
+	template <typename T, typename TEnableIf<TAnd<TNot<TIsEnumClass<T>>, TNot<TModels<CBaseStructureProvider, T>>, TNot<TIsDerivedFrom<std::remove_pointer_t<T>, UObject>>>::Value>::Type* = nullptr>
+	static EPropertyBagResult CheckType(FPropertyBagPropertyDesc InDesc, const T InValue)
+	{
+		if (!InDesc.ContainerTypes.IsEmpty())
+		{
+			return EPropertyBagResult::TypeMismatch;
+		}
+
+		// TODO: consider adding generic traits/concepts to PropertyBag for mapping from T -> EPropertyBagPropertyType
+		// This generic mechanism should consider the entirety of the PropertyBag interface and be applied consistently.
+		if ((std::is_same_v<T, bool> && InDesc.ValueType == EPropertyBagPropertyType::Bool)
+			|| (std::is_same_v<T, uint8> && InDesc.ValueType == EPropertyBagPropertyType::Byte)
+			|| (std::is_same_v<T, int32> && InDesc.ValueType == EPropertyBagPropertyType::Int32)
+			|| (std::is_same_v<T, int64> && InDesc.ValueType == EPropertyBagPropertyType::Int64)
+			|| (std::is_same_v<T, uint32> && InDesc.ValueType == EPropertyBagPropertyType::UInt32)
+			|| (std::is_same_v<T, uint64> && InDesc.ValueType == EPropertyBagPropertyType::UInt64)
+			|| (std::is_same_v<T, float> && InDesc.ValueType == EPropertyBagPropertyType::Float)
+			|| (std::is_same_v<T, double> && InDesc.ValueType == EPropertyBagPropertyType::Double)
+			|| (std::is_same_v<T, FName> && InDesc.ValueType == EPropertyBagPropertyType::Name)
+			|| (std::is_same_v<T, FString> && InDesc.ValueType == EPropertyBagPropertyType::String)
+			|| (std::is_same_v<T, FText> && InDesc.ValueType == EPropertyBagPropertyType::Text)
+			|| (std::is_same_v<T, FSoftObjectPtr> && InDesc.ValueType == EPropertyBagPropertyType::SoftObject))
+		{
+			return EPropertyBagResult::Success;
+		}
+
+		return EPropertyBagResult::TypeMismatch;
+	}
+
+	template <typename T, typename ReturnType>
+	static ReturnType Call(TFunction<ReturnType(const void*)>&& Func, const FPropertyBagPropertyDesc& InDesc, const T& InValue)
+	{
+		if constexpr (TIsDerivedFrom<T, FSoftObjectPath>::Value || TIsDerivedFrom<std::remove_pointer_t<T>, UObject>::Value)
+		{
+			if (InDesc.ValueType == EPropertyBagPropertyType::SoftObject || InDesc.ValueType == EPropertyBagPropertyType::SoftClass)
+			{
+				return Call(MoveTemp(Func), InDesc, FSoftObjectPtr(InValue));
+			}
+		}
+
+		return Func(&InValue);
+	}
+
+	template <typename T>
+	static void Convert(const FPropertyBagPropertyDesc& InDesc, void* InValue, T& OutValue)
+	{
+		if constexpr (TIsDerivedFrom<std::remove_pointer_t<T>, UObject>::Value)
+		{
+			static_assert(std::is_pointer_v<T>);
+
+			if (InDesc.ValueType == EPropertyBagPropertyType::SoftObject || InDesc.ValueType == EPropertyBagPropertyType::SoftClass)
+			{
+				const FObjectPropertyBase* ObjectProperty = CastFieldChecked<FObjectPropertyBase>(InDesc.CachedProperty);
+				check(ObjectProperty->PropertyClass);
+
+				OutValue = Cast<std::remove_pointer_t<T>>(ObjectProperty->GetObjectPropertyValue(InValue));
+				return;
+			}
+		}
+		else if constexpr (TIsDerivedFrom<T, FSoftObjectPath>::Value)
+		{
+			if (InDesc.ValueType == EPropertyBagPropertyType::SoftObject || InDesc.ValueType == EPropertyBagPropertyType::SoftClass)
+			{
+				const FSoftObjectProperty* SoftObjectProperty = CastFieldChecked<FSoftObjectProperty>(InDesc.CachedProperty);
+				check(SoftObjectProperty->PropertyClass);
+
+				if constexpr (TIsDerivedFrom<T, FSoftClassPath>::Value)
+				{
+					OutValue = FSoftClassPath(SoftObjectProperty->GetPropertyValue(InValue).ToSoftObjectPath().ToString());
+				}
+				else
+				{
+					OutValue = SoftObjectProperty->GetPropertyValue(InValue).ToSoftObjectPath();
+				}
+				return;
+			}
+		}
+
+		OutValue = *reinterpret_cast<T*>(InValue);
+	}
+
+	template <typename KeyType, typename ValueType>
+	EPropertyBagResult AddPair(const FPropertyBagPropertyDesc& InKeyDesc, const KeyType& InKey, const FPropertyBagPropertyDesc& InValueDesc, const ValueType& InValue)
+	{
+		if constexpr (
+			(TIsDerivedFrom<KeyType, FSoftObjectPath>::Value || TIsDerivedFrom<std::remove_pointer_t<KeyType>, UObject>::Value) &&
+			(TIsDerivedFrom<ValueType, FSoftObjectPath>::Value || TIsDerivedFrom<std::remove_pointer_t<ValueType>, UObject>::Value)
+		)
+		{
+			if ((InKeyDesc.ValueType == EPropertyBagPropertyType::SoftObject || InKeyDesc.ValueType == EPropertyBagPropertyType::SoftClass) &&
+				(InValueDesc.ValueType == EPropertyBagPropertyType::SoftObject || InValueDesc.ValueType == EPropertyBagPropertyType::SoftClass))
+			{
+				return AddPair(InKeyDesc, FSoftObjectPtr(InKey), InValueDesc, FSoftObjectPtr(InValue));
+			}
+		}
+		else if constexpr (TIsDerivedFrom<KeyType, FSoftObjectPath>::Value || TIsDerivedFrom<std::remove_pointer_t<KeyType>, UObject>::Value)
+		{
+			if (InKeyDesc.ValueType == EPropertyBagPropertyType::SoftObject || InKeyDesc.ValueType == EPropertyBagPropertyType::SoftClass)
+			{
+				return AddPair(InKeyDesc, FSoftObjectPtr(InKey), InValueDesc, InValue);
+			}
+		}
+		else if constexpr (TIsDerivedFrom<ValueType, FSoftObjectPath>::Value || TIsDerivedFrom<std::remove_pointer_t<ValueType>, UObject>::Value)
+		{
+			if (InValueDesc.ValueType == EPropertyBagPropertyType::SoftObject || InValueDesc.ValueType == EPropertyBagPropertyType::SoftClass)
+			{
+				return AddPair(InKeyDesc, InKey, InValueDesc, FSoftObjectPtr(InValue));
+			}
+		}
+
+		FScriptMapHelper::AddPair(&InKey, &InValue);
+		return EPropertyBagResult::Success;
+	}
+
+	FPropertyBagPropertyDesc KeyDesc;
+	FPropertyBagPropertyDesc ValueDesc;
+
+};
