[StructUtils] Fix the ID once the struct is created. The ID is part of the hashing, and 2 new structs with the same properties should share the same UPropertyBag instance. Introduced with CL 49337770.
#Synced-CL 49284918

diff --git a/Engine/Source/Developer/StructUtilsTestSuite/Private/PropertyBagTest.cpp b/Engine/Source/Developer/StructUtilsTestSuite/Private/PropertyBagTest.cpp
--- a/Engine/Source/Developer/StructUtilsTestSuite/Private/PropertyBagTest.cpp
+++ b/Engine/Source/Developer/StructUtilsTestSuite/Private/PropertyBagTest.cpp
@@ -1600,6 +1600,7 @@ struct FTest_SameBag : FAITestBase
 	{
 		static const FName TemperatureName(TEXT("Temperature"));
 		static const FName CountName(TEXT("Count"));
+		static const FName AmountName(TEXT("Amount"));
 
 		FInstancedPropertyBag BagA;
 		BagA.AddProperties({
@@ -1617,6 +1618,22 @@ struct FTest_SameBag : FAITestBase
 		// Same descriptors should result in same bag struct
 		AITEST_TRUE(TEXT("Property bags should match"), BagA.GetPropertyBagStruct() == BagB.GetPropertyBagStruct());
 
+		// Add a new property. They should not match.
+		BagB.AddProperties({ {AmountName, EPropertyBagPropertyType::UInt64} });
+		AITEST_FALSE(TEXT("Property bags not should match after adding a new property"), BagA.GetPropertyBagStruct() == BagB.GetPropertyBagStruct());
+
+		// Remove previous property. They should now matches.
+		BagB.RemovePropertiesByName({AmountName});
+		AITEST_TRUE(TEXT("Property bags should match after removing the new property"), BagA.GetPropertyBagStruct() == BagB.GetPropertyBagStruct());
+
+		// Remove exiting property. They should not matches.
+		BagB.RemovePropertiesByName({CountName});
+		AITEST_FALSE(TEXT("Property bags should not match after removing a property"), BagA.GetPropertyBagStruct() == BagB.GetPropertyBagStruct());
+
+		// Re-adding the property. They should now matches.
+		BagB.AddProperties({ {CountName, EPropertyBagPropertyType::Int32} });
+		AITEST_TRUE(TEXT("Property bags should match after readding the property"), BagA.GetPropertyBagStruct() == BagB.GetPropertyBagStruct());
+
 		return true;
 	}
 };
diff --git a/Engine/Source/Runtime/CoreUObject/Private/StructUtils/PropertyBag.cpp b/Engine/Source/Runtime/CoreUObject/Private/StructUtils/PropertyBag.cpp
--- a/Engine/Source/Runtime/CoreUObject/Private/StructUtils/PropertyBag.cpp
+++ b/Engine/Source/Runtime/CoreUObject/Private/StructUtils/PropertyBag.cpp
@@ -4264,12 +4264,6 @@ TArray<FPropertyBagPropertyDesc> UPropertyBag::SanitizePropertyDescs(const TArra
 		{
 			CheckType(Desc.Name, Desc.KeyType, Desc.KeyTypeObject);
 		}
-
-		// Fix ID
-		if (!Desc.ID.IsValid())
-		{
-			Desc.ID = FGuid::NewGuid();
-		}
 	}
 
 	// Remove properties with same name
@@ -4311,6 +4305,15 @@ const UPropertyBag* UPropertyBag::GetOrCreateFromDescs(const TConstArrayView<FPr
 
 	NewBag->PropertyDescs = MoveTemp(PropertyDescs);
 
+	// Fix IDs
+	for (FPropertyBagPropertyDesc& Desc : PropertyDescs)
+	{
+		if (!Desc.ID.IsValid())
+		{
+			Desc.ID = FGuid::NewGuid();
+		}
+	}
+
 	// Add properties (AddCppProperty() adds them backwards in the linked list)
 	for (int32 DescIndex = NewBag->PropertyDescs.Num() - 1; DescIndex >= 0; DescIndex--)
 	{
