[StructUtils] Update TSharedStruct::Make function. They used to not compile. Add unittest.

#rnx

#Synced-CL 50826644

diff --git a/Engine/Source/Developer/StructUtilsTestSuite/Private/StructSharedStructTest.cpp b/Engine/Source/Developer/StructUtilsTestSuite/Private/StructSharedStructTest.cpp
--- a/Engine/Source/Developer/StructUtilsTestSuite/Private/StructSharedStructTest.cpp
+++ b/Engine/Source/Developer/StructUtilsTestSuite/Private/StructSharedStructTest.cpp
@@ -17,7 +17,7 @@ namespace FStructSharedStructTest
 			{
 				FSharedStruct SharedStruct = FSharedStruct::Make<FTestStructSimpleNonZeroDefault>();
 
-				AITEST_EQUAL("FSharedStruct default initalized from Make should have same value as default constructed", FTestStructSimpleNonZeroDefault(), SharedStruct.Get<FTestStructSimpleNonZeroDefault>());
+				AITEST_EQUAL("FSharedStruct default initialized from Make should have same value as default constructed", FTestStructSimpleNonZeroDefault(), SharedStruct.Get<FTestStructSimpleNonZeroDefault>());
 			}
 
 			{
@@ -55,6 +55,22 @@ namespace FStructSharedStructTest
 
 				AITEST_EQUAL("FSharedStruct initiazed from Make should have value reflecting TArgs", Val, SharedStruct.Get<FTestStructSimple>().Float);
 			}
+			{
+				const float Val = 99.f;
+				FTestStructSimple Simple(Val);
+
+				TSharedStruct<FTestStructSimple> SharedStruct = TSharedStruct<FTestStructSimple>::Make<FTestStructSimple>(Simple);
+
+				AITEST_EQUAL("TSharedStruct initiazed from Make should have value reflecting TArgs", Val, SharedStruct.Get<FTestStructSimple>().Float);
+			}
+			{
+				const float Val = 99.f;
+				FTestStructSimple Simple(Val);
+
+				TSharedStruct<FTestStructSimple> SharedStruct = TSharedStruct<FTestStructSimple>::Make<FTestStructSimple>(99);
+
+				AITEST_EQUAL("TSharedStruct initiazed from Make should have value reflecting TArgs", Val, SharedStruct.Get<FTestStructSimple>().Float);
+			}
 
 			return true;
 		}
@@ -116,7 +132,7 @@ namespace FStructSharedStructTest
 			{
 				FConstSharedStruct SharedStruct = FConstSharedStruct::Make<FTestStructSimpleNonZeroDefault>();
 
-				AITEST_EQUAL("FConstSharedStruct default initalized from Make should have same value as default constructed", FTestStructSimpleNonZeroDefault(), SharedStruct.Get<const FTestStructSimpleNonZeroDefault>());
+				AITEST_EQUAL("FConstSharedStruct default initialized from Make should have same value as default constructed", FTestStructSimpleNonZeroDefault(), SharedStruct.Get<const FTestStructSimpleNonZeroDefault>());
 			}
 
 			{
diff --git a/Engine/Source/Runtime/CoreUObject/Public/StructUtils/SharedStruct.h b/Engine/Source/Runtime/CoreUObject/Public/StructUtils/SharedStruct.h
--- a/Engine/Source/Runtime/CoreUObject/Public/StructUtils/SharedStruct.h
+++ b/Engine/Source/Runtime/CoreUObject/Public/StructUtils/SharedStruct.h
@@ -445,7 +445,9 @@ struct [[nodiscard]] TSharedStruct
 		requires (std::is_base_of_v<BaseStructT, std::decay_t<T>>)
 	static TSharedStruct Make()
 	{
-		return FSharedStruct::Make<T>();
+		TSharedStruct Result;
+		Result.SharedStruct = FSharedStruct::Make<T>();
+		return Result;
 	}
 
 	/** Creates a new TSharedStruct from templated struct instance. This will create a new instance of the shared struct memory. */
@@ -453,14 +455,18 @@ struct [[nodiscard]] TSharedStruct
 		requires (std::is_base_of_v<BaseStructT, std::decay_t<T>>)
 	static TSharedStruct Make(const T& Struct)
 	{
-		return FSharedStruct::Make<T>(Struct);
+		TSharedStruct Result;
+		Result.SharedStruct = FSharedStruct::Make<T>(Struct);
+		return Result;
 	}
 
 	/** Creates a new TSharedStruct from the templated type and forward all arguments to constructor. This will create a new instance of the shared struct memory. */
 	template< typename... TArgs>
 	static TSharedStruct Make(TArgs&&... InArgs)
 	{
-		return FSharedStruct::Make<BaseStructT>(Forward<TArgs>(InArgs)...);
+		TSharedStruct Result;
+		Result.SharedStruct = FSharedStruct::Make<BaseStructT>(Forward<TArgs>(InArgs)...);
+		return Result;
 	}
 
 	/** Returns reference to the struct, this getter assumes that all data is valid. */
