[StructUtils] Fix TSharedStruct::operator==. Add new unittest for TSharedStruct and TConstSharedStruct.

#Synced-CL 50969155

diff --git a/Engine/Source/Developer/StructUtilsTestSuite/Private/StructSharedStructTest.cpp b/Engine/Source/Developer/StructUtilsTestSuite/Private/StructSharedStructTest.cpp
--- a/Engine/Source/Developer/StructUtilsTestSuite/Private/StructSharedStructTest.cpp
+++ b/Engine/Source/Developer/StructUtilsTestSuite/Private/StructSharedStructTest.cpp
@@ -26,7 +26,7 @@ namespace FStructSharedStructTest
 
 				FSharedStruct SharedStruct = FSharedStruct::Make(Simple);
 
-				AITEST_EQUAL("FSharedStruct initiazed from Make should have value of FTestStructSimple its initiliazed from", Val, SharedStruct.Get<FTestStructSimple>().Float);
+				AITEST_EQUAL("FSharedStruct initialized from Make should have value of FTestStructSimple its initiliazed from", Val, SharedStruct.Get<FTestStructSimple>().Float);
 			}
 
 			{
@@ -35,7 +35,7 @@ namespace FStructSharedStructTest
 
 				FSharedStruct SharedStruct = FSharedStruct::Make(FTestStructSimple::StaticStruct(), (uint8*)&Simple);
 
-				AITEST_EQUAL("FSharedStruct initiazed from Make should have value of StaticStruct and memory its initiliazed from", Val, SharedStruct.Get<FTestStructSimple>().Float);
+				AITEST_EQUAL("FSharedStruct initialized from Make should have value of StaticStruct and memory its initiliazed from", Val, SharedStruct.Get<FTestStructSimple>().Float);
 			}
 
 			{
@@ -44,7 +44,7 @@ namespace FStructSharedStructTest
 				FStructView StructView = FStructView::Make(Simple);
 				FSharedStruct SharedStruct = FSharedStruct::Make(StructView);
 
-				AITEST_EQUAL("FSharedStruct initiazed from Make should have value of FStructView its initiliazed from", Val, SharedStruct.Get<const FTestStructSimple>().Float);
+				AITEST_EQUAL("FSharedStruct initialized from Make should have value of FStructView its initiliazed from", Val, SharedStruct.Get<const FTestStructSimple>().Float);
 			}
 
 			{
@@ -53,23 +53,7 @@ namespace FStructSharedStructTest
 
 				FSharedStruct SharedStruct = FSharedStruct::Make<FTestStructSimple>(99);
 
-				AITEST_EQUAL("FSharedStruct initiazed from Make should have value reflecting TArgs", Val, SharedStruct.Get<FTestStructSimple>().Float);
-			}
-			{
-				const float Val = 99.f;
-				FTestStructSimple Simple(Val);
-
-				TSharedStruct<FTestStructSimple> SharedStruct = TSharedStruct<FTestStructSimple>::Make<FTestStructSimple>(Simple);
-
-				AITEST_EQUAL("TSharedStruct initiazed from Make should have value reflecting TArgs", Val, SharedStruct.Get<FTestStructSimple>().Float);
-			}
-			{
-				const float Val = 99.f;
-				FTestStructSimple Simple(Val);
-
-				TSharedStruct<FTestStructSimple> SharedStruct = TSharedStruct<FTestStructSimple>::Make<FTestStructSimple>(99);
-
-				AITEST_EQUAL("TSharedStruct initiazed from Make should have value reflecting TArgs", Val, SharedStruct.Get<FTestStructSimple>().Float);
+				AITEST_EQUAL("FSharedStruct initialized from Make should have value reflecting TArgs", Val, SharedStruct.Get<FTestStructSimple>().Float);
 			}
 
 			return true;
@@ -119,6 +103,23 @@ namespace FStructSharedStructTest
 				AITEST_TRUE("FSharedStruct created to a specific struct type should be IsValid()", SharedStruct.IsValid());
 			}
 
+			{
+				FSharedStruct SharedStruct1 = FSharedStruct::Make<FTestStructSimple>();
+				FSharedStruct SharedStruct2 = FSharedStruct::Make<FTestStructSimple>();
+				const bool bEqual = SharedStruct1 == SharedStruct2;
+				AITEST_FALSE("FSharedStruct operator== should compare what is shared and not what it contains", bEqual);
+			}
+
+			{
+				FSharedStruct SharedStruct1 = FSharedStruct::Make<FTestStructSimple>();
+				FSharedStruct SharedStruct2 = FSharedStruct::Make<FTestStructSimple>();
+				TArray<FSharedStruct> Array;
+				Array.Add(SharedStruct1);
+
+				AITEST_TRUE("FSharedStruct should support TArray", Array.Contains(SharedStruct1));
+				AITEST_FALSE("FSharedStruct should support TArray", Array.Contains(SharedStruct2));
+			}
+
 			return true;
 		}
 	};
@@ -141,7 +142,7 @@ namespace FStructSharedStructTest
 
 				FConstSharedStruct SharedStruct = FConstSharedStruct::Make(Simple);
 
-				AITEST_EQUAL("FConstSharedStruct initiazed from Make should have value of FTestStructSimple its initiliazed from", Val, SharedStruct.Get<const FTestStructSimple>().Float);
+				AITEST_EQUAL("FConstSharedStruct initialized from Make should have value of FTestStructSimple its initiliazed from", Val, SharedStruct.Get<const FTestStructSimple>().Float);
 			}
 
 			{
@@ -150,7 +151,7 @@ namespace FStructSharedStructTest
 
 				FConstSharedStruct ConstSharedStruct = FConstSharedStruct::Make(FTestStructSimple::StaticStruct(), (uint8*)&Simple);
 
-				AITEST_EQUAL("FConstSharedStruct initiazed from Make should have value of StaticStruct and memory its initiliazed from", Val, ConstSharedStruct.Get<const FTestStructSimple>().Float);
+				AITEST_EQUAL("FConstSharedStruct initialized from Make should have value of StaticStruct and memory its initiliazed from", Val, ConstSharedStruct.Get<const FTestStructSimple>().Float);
 			}
 
 			{
@@ -159,7 +160,7 @@ namespace FStructSharedStructTest
 				FStructView StructView = FStructView::Make(Simple);
 				FConstSharedStruct ConstSharedStruct = FConstSharedStruct::Make(StructView);
 
-				AITEST_EQUAL("FConstSharedStruct initiazed from Make should have value of FStructView its initiliazed from", Val, ConstSharedStruct.Get<const FTestStructSimple>().Float);
+				AITEST_EQUAL("FConstSharedStruct initialized from Make should have value of FStructView its initiliazed from", Val, ConstSharedStruct.Get<const FTestStructSimple>().Float);
 			}
 
 			{
@@ -168,7 +169,7 @@ namespace FStructSharedStructTest
 
 				FConstSharedStruct ConstSharedStruct = FConstSharedStruct::Make<FTestStructSimple>(99);
 
-				AITEST_EQUAL("FConstSharedStruct initiazed from Make should have value reflecting TArgs", Val, ConstSharedStruct.Get<const FTestStructSimple>().Float);
+				AITEST_EQUAL("FConstSharedStruct initialized from Make should have value reflecting TArgs", Val, ConstSharedStruct.Get<const FTestStructSimple>().Float);
 			}
 
 			return true;
@@ -234,12 +235,191 @@ namespace FStructSharedStructTest
 				AITEST_TRUE("FConstSharedStruct created to a specific struct type should be IsValid()", ConstSharedStruct.IsValid());
 			}
 
+			{
+				FConstSharedStruct ConstSharedStruct1 = FConstSharedStruct::Make<FTestStructSimple>();
+				FConstSharedStruct ConstSharedStruct2 = FConstSharedStruct::Make<FTestStructSimple>();
+				const bool bEqual = ConstSharedStruct1 == ConstSharedStruct2;
+				AITEST_FALSE("FConstSharedStruct operator== should compare what is shared and not what it contains", bEqual);
+			}
+
+			{
+				FConstSharedStruct ConstSharedStruct1 = FConstSharedStruct::Make<FTestStructSimple>();
+				FConstSharedStruct ConstSharedStruct2 = FConstSharedStruct::Make<FTestStructSimple>();
+				TArray<FConstSharedStruct> Array;
+				Array.Add(ConstSharedStruct1);
+
+				AITEST_TRUE("FConstSharedStruct should support TArray", Array.Contains(ConstSharedStruct1));
+				AITEST_FALSE("FConstSharedStruct should support TArray", Array.Contains(ConstSharedStruct2));
+			}
+
 			return true;
 		}
 	};
 
 	IMPLEMENT_AI_INSTANT_TEST(FTest_ConstSharedStructBasic, "System.StructUtils.SharedStruct.Const.Basic");
 
+	struct FTest_TSharedStructBasic : FAITestBase
+	{
+		virtual bool InstantTest() override
+		{
+			{
+				TSharedStruct<FTestStructSimple> SharedStruct = TSharedStruct<FTestStructSimple>::Make();
+				TSharedStruct<FTestStructSimple> SharedStruct2(SharedStruct);
+
+				AITEST_EQUAL("TSharedStruct and TSharedStruct should be equal from copy construction", SharedStruct, SharedStruct2);
+			}
+
+			{
+				TSharedStruct<FTestStructSimple> SharedStruct = TSharedStruct<FTestStructSimple>::Make();
+				TSharedStruct<FTestStructSimple> SharedStruct2;
+				
+				SharedStruct2 = SharedStruct;
+
+				AITEST_EQUAL("TSharedStruct and TSharedStruct should be equal from copy asignment", SharedStruct, SharedStruct2);
+			}
+			{
+				constexpr float Val = 99.f;
+
+				TSharedStruct<FTestStructSimple> SharedStruct1 = TSharedStruct<FTestStructSimple>::Make();
+				TSharedStruct<FTestStructSimple> SharedStruct2(SharedStruct1);
+
+				SharedStruct1.Get<FTestStructSimple>().Float = Val;
+				AITEST_EQUAL("FSharedStruct and TSharedStruct should reflect the same changes to shared FTestStructSimple", SharedStruct1.Get().Float, SharedStruct2.Get().Float);
+			}
+
+			{
+				TSharedStruct<FTestStructSimple> SharedStructTemplated;
+				AITEST_FALSE("Default constructed TSharedStruct should IsValid() == false", SharedStructTemplated.IsValid());
+			}
+
+			{
+				TSharedStruct<FTestStructSimple> SharedStructTemplated = TSharedStruct<FTestStructSimple>::Make();
+				AITEST_TRUE("TSharedStruct created to a specific struct type should be IsValid()", SharedStructTemplated.IsValid());
+			}
+
+			{
+				TSharedStruct<FTestStructSimple> SharedStructTemplated1 = TSharedStruct<FTestStructSimple>::Make();
+				TSharedStruct<FTestStructSimple> SharedStructTemplated2 = TSharedStruct<FTestStructSimple>::Make();
+				const bool bEqual = SharedStructTemplated1 == SharedStructTemplated2;
+				AITEST_FALSE("TSharedStruct operator== should compare what is shared and not what it contains", bEqual);
+			}
+
+			{
+				TSharedStruct<FTestStructSimple> SharedStructTemplated1 = TSharedStruct<FTestStructSimple>::Make();
+				TSharedStruct<FTestStructSimple> SharedStructTemplated2 = TSharedStruct<FTestStructSimple>::Make();
+				TArray<TSharedStruct<FTestStructSimple>> Array;
+				Array.Add(SharedStructTemplated1);
+
+				AITEST_TRUE("TSharedStruct should support TArray", Array.Contains(SharedStructTemplated1));
+				AITEST_FALSE("TSharedStruct should support TArray", Array.Contains(SharedStructTemplated2));
+			}
+
+			return true;
+		}
+	};
+
+	IMPLEMENT_AI_INSTANT_TEST(FTest_TSharedStructBasic, "System.StructUtils.SharedStruct.Template.Basic");
+
+	struct FTest_TConstSharedStructBasic : FAITestBase
+	{
+		virtual bool InstantTest() override
+		{
+			{
+				TConstSharedStruct<FTestStructSimple> SharedStruct1 = TConstSharedStruct<FTestStructSimple>::Make();
+				TConstSharedStruct<FTestStructSimple> SharedStruct2(SharedStruct1);
+
+				AITEST_EQUAL("TConstSharedStruct and TConstSharedStruct should be equal from copy construction", SharedStruct1, SharedStruct2);
+			}
+
+			{
+				TSharedStruct<FTestStructSimple> SharedStruct1 = TSharedStruct<FTestStructSimple>::Make();
+				TConstSharedStruct<FTestStructSimple> SharedStruct2(SharedStruct1);
+
+				const bool bEqual = SharedStruct1 == SharedStruct2;
+				const bool bNotEqual = SharedStruct1 != SharedStruct2;
+				AITEST_TRUE("TConstSharedStruct and TSharedStruct should be equal from copy construction", bEqual);
+				AITEST_FALSE("TConstSharedStruct and TSharedStruct should be equal from copy construction", bNotEqual);
+			}
+
+			{
+				TConstSharedStruct<FTestStructSimple> SharedStruct1 = TConstSharedStruct<FTestStructSimple>::Make();
+				TConstSharedStruct<FTestStructSimple> SharedStruct2;
+				
+				SharedStruct2 = SharedStruct1;
+
+				AITEST_EQUAL("TConstSharedStruct and TConstSharedStruct should be equal from copy asignment", SharedStruct1, SharedStruct2);
+			}
+
+			{
+				TSharedStruct<FTestStructSimple> SharedStruct1 = TSharedStruct<FTestStructSimple>::Make();
+				TConstSharedStruct<FTestStructSimple> SharedStruct2;
+				
+				SharedStruct2 = SharedStruct1;
+				const bool bEqual = SharedStruct1 == SharedStruct2;
+				const bool bNotEqual = SharedStruct1 != SharedStruct2;
+				AITEST_TRUE("TSharedStruct and TConstSharedStruct should be equal from copy asignment", bEqual);
+				AITEST_FALSE("TSharedStruct and TConstSharedStruct should be equal from copy asignment", bNotEqual);
+			}
+
+			{
+				constexpr float Val = 99.f;
+
+				TConstSharedStruct<FTestStructSimple> SharedStruct1 = TConstSharedStruct<FTestStructSimple>::Make(Val);
+				TConstSharedStruct<FTestStructSimple> SharedStruct2(SharedStruct1);
+
+				AITEST_EQUAL("TConstSharedStruct and TConstSharedStruct should reflect the same changes to shared FTestStructSimple", SharedStruct1.Get().Float, SharedStruct2.Get().Float);
+				AITEST_EQUAL("TConstSharedStruct and TConstSharedStruct should reflect the same changes to shared FTestStructSimple", SharedStruct1.GetPtr()->Float, SharedStruct2.GetPtr()->Float);
+				AITEST_EQUAL("TConstSharedStruct and TConstSharedStruct should reflect the same changes to shared FTestStructSimple", SharedStruct1.Get<FTestStructSimple>().Float, SharedStruct2.Get<const FTestStructSimple>().Float);
+				AITEST_EQUAL("TConstSharedStruct and TConstSharedStruct should reflect the same changes to shared FTestStructSimple", SharedStruct1.GetPtr<FTestStructSimple>()->Float, SharedStruct2.GetPtr<const FTestStructSimple>()->Float);
+			}
+
+			{
+				TConstSharedStruct<FTestStructSimple> SharedStructTemplated;
+				AITEST_FALSE("Default constructed TConstSharedStruct should IsValid() == false", SharedStructTemplated.IsValid());
+			}
+
+			{
+				TConstSharedStruct<FTestStructSimple> SharedStructTemplated = TConstSharedStruct<FTestStructSimple>::Make();
+				AITEST_TRUE("TConstSharedStruct created to a specific struct type should be IsValid()", SharedStructTemplated.IsValid());
+			}
+
+			{
+				TConstSharedStruct<FTestStructSimple> SharedStructTemplated1 = TConstSharedStruct<FTestStructSimple>::Make();
+				TConstSharedStruct<FTestStructSimple> SharedStructTemplated2 = TConstSharedStruct<FTestStructSimple>::Make();
+				const bool bEqual = SharedStructTemplated1 == SharedStructTemplated2;
+				AITEST_FALSE("TConstSharedStruct operator== should compare what is shared and not what it contains", bEqual);
+			}
+
+			{
+				TSharedStruct<FTestStructSimple> SharedStructTemplated1 = TSharedStruct<FTestStructSimple>::Make();
+				TConstSharedStruct<FTestStructSimple> SharedStructTemplated2 = TConstSharedStruct<FTestStructSimple>::Make();
+				const bool bEqual = SharedStructTemplated1 == SharedStructTemplated2;
+				AITEST_FALSE("TConstSharedStruct operator== should compare what is shared and not what it contains", bEqual);
+			}
+
+			{
+				TConstSharedStruct<FTestStructSimple> SharedStructTemplated1 = TConstSharedStruct<FTestStructSimple>::Make();
+				TSharedStruct<FTestStructSimple> SharedStructTemplated2 = TSharedStruct<FTestStructSimple>::Make();
+				const bool bEqual = SharedStructTemplated1 == SharedStructTemplated2;
+				AITEST_FALSE("TConstSharedStruct operator== should compare what is shared and not what it contains", bEqual);
+			}
+
+			{
+				TConstSharedStruct<FTestStructSimple> SharedStructTemplated1 = TConstSharedStruct<FTestStructSimple>::Make();
+				TConstSharedStruct<FTestStructSimple> SharedStructTemplated2 = TConstSharedStruct<FTestStructSimple>::Make();
+				TArray<TConstSharedStruct<FTestStructSimple>> Array;
+				Array.Add(SharedStructTemplated1);
+
+				AITEST_TRUE("TSharedStruct should support TArray", Array.Contains(SharedStructTemplated1));
+				AITEST_FALSE("TSharedStruct should support TArray", Array.Contains(SharedStructTemplated2));
+			}
+
+			return true;
+		}
+	};
+
+	IMPLEMENT_AI_INSTANT_TEST(FTest_TConstSharedStructBasic, "System.StructUtils.SharedStruct.Template.Const.Basic");
+
 	struct FTest_SharedStructDestruction : FAITestBase
 	{
 		virtual bool InstantTest() override
diff --git a/Engine/Source/Runtime/CoreUObject/Public/StructUtils/SharedStruct.h b/Engine/Source/Runtime/CoreUObject/Public/StructUtils/SharedStruct.h
--- a/Engine/Source/Runtime/CoreUObject/Public/StructUtils/SharedStruct.h
+++ b/Engine/Source/Runtime/CoreUObject/Public/StructUtils/SharedStruct.h
@@ -494,16 +494,30 @@ struct [[nodiscard]] TSharedStruct
 	/** Comparison operators. Note: it does not compare the internal structure itself */
 	template<typename T = BaseStructT>
 		requires (std::is_base_of_v<BaseStructT, std::decay_t<T>>)
-	bool operator==(const T& Other) const
+	bool operator==(const TSharedStruct<T>& Other) const
 	{
-		return SharedStruct.operator==(Other);
+		return SharedStruct.operator==(Other.SharedStruct);
 	}
 
 	template<typename T = BaseStructT>
 		requires (std::is_base_of_v<BaseStructT, std::decay_t<T>>)
-	bool operator!=(const T& Other) const
+	bool operator!=(const TSharedStruct<T>& Other) const
 	{
-		return SharedStruct.operator!=(Other);
+		return SharedStruct.operator!=(Other.SharedStruct);
+	}
+
+	template<typename T = BaseStructT>
+		requires (std::is_base_of_v<BaseStructT, std::decay_t<T>>)
+	bool operator==(const TConstSharedStruct<T>& Other) const
+	{
+		return SharedStruct.operator==(Other.ConstSharedStruct);
+	}
+
+	template<typename T = BaseStructT>
+		requires (std::is_base_of_v<BaseStructT, std::decay_t<T>>)
+	bool operator!=(const TConstSharedStruct<T>& Other) const
+	{
+		return SharedStruct.operator!=(Other.ConstSharedStruct);
 	}
 
 private:
@@ -518,6 +532,7 @@ private:
 	FSharedStruct SharedStruct;
 
 	template <typename U> friend struct TSharedStruct;
+	template <typename U> friend struct TConstSharedStruct;
 };
 
 template<>
@@ -680,16 +695,14 @@ struct [[nodiscard]] FConstSharedStruct
 
 	/** Returns const reference to the struct, this getter assumes that all data is valid. */
 	template<typename T>
-		requires (std::is_const_v<T>)
-	constexpr T& Get() const
+	const T& Get() const
 	{
 		return UE::StructUtils::GetStructRef<T>(GetScriptStruct(), GetMemory());
 	}
 
 	/** Returns const pointer to the struct, or nullptr if cast is not valid. */
 	template<typename T>
-		requires (std::is_const_v<T>)
-	constexpr T* GetPtr() const
+	const T* GetPtr() const
 	{
 		return UE::StructUtils::GetStructPtr<T>(GetScriptStruct(), GetMemory());
 	}
@@ -755,8 +768,8 @@ struct [[nodiscard]] TConstSharedStruct
 
 	template<typename T = BaseStructT>
 		requires (std::is_base_of_v<BaseStructT, std::decay_t<T>>)
-	TConstSharedStruct(const TSharedStruct<T>& ConstSharedStruct)
-		: ConstSharedStruct(ConstSharedStruct.SharedStruct)
+	TConstSharedStruct(const TSharedStruct<T>& SharedStruct)
+		: ConstSharedStruct(SharedStruct.SharedStruct)
 	{
 	}
 
@@ -769,8 +782,8 @@ struct [[nodiscard]] TConstSharedStruct
 
 	template<typename T = BaseStructT>
 		requires (std::is_base_of_v<BaseStructT, std::decay_t<T>>)
-	TConstSharedStruct(TSharedStruct<T>&& ConstSharedStruct)
-		: ConstSharedStruct(MoveTemp(ConstSharedStruct.SharedStruct))
+	TConstSharedStruct(TSharedStruct<T>&& SharedStruct)
+		: ConstSharedStruct(MoveTemp(SharedStruct.SharedStruct))
 	{
 	}
 
@@ -778,7 +791,7 @@ struct [[nodiscard]] TConstSharedStruct
 		requires (std::is_base_of_v<BaseStructT, std::decay_t<T>>)
 	TConstSharedStruct& operator=(const TConstSharedStruct<T>& Other)
 	{
-		if (this == &Other)
+		if (this != &Other)
 		{
 			ConstSharedStruct = Other.ConstSharedStruct;
 		}
@@ -789,10 +802,7 @@ struct [[nodiscard]] TConstSharedStruct
 		requires (std::is_base_of_v<BaseStructT, std::decay_t<T>>)
 	TConstSharedStruct& operator=(const TSharedStruct<T>& SharedStruct)
 	{
-		if (this != &SharedStruct)
-		{
-			ConstSharedStruct = SharedStruct.SharedStruct;
-		}
+		ConstSharedStruct = SharedStruct.SharedStruct;
 		return *this;
 	}
 
@@ -836,37 +846,27 @@ struct [[nodiscard]] TConstSharedStruct
 		ConstSharedStruct.InitializeAs<T>(Forward<TArgs>(InArgs)...);
 	}
 
-	/** Creates a new TSharedStruct from a templated struct type. This will create a new instance of the shared struct memory. */
+	/** Creates a new TConstSharedStruct from a templated struct type. This will create a new instance of the shared struct memory. */
 	template<typename T = BaseStructT>
 		requires (std::is_base_of_v<BaseStructT, std::decay_t<T>>)
 	static TConstSharedStruct Make()
 	{
-		FConstSharedStruct SharedStruct;
-		SharedStruct.InitializeAs<BaseStructT>();
+		TConstSharedStruct SharedStruct;
+		SharedStruct.Initialize<BaseStructT>();
 		return SharedStruct;
 	}
 
-	/** Creates a new TSharedStruct from templated struct instance. This will create a new instance of the shared struct memory. */
+	/** Creates a new TConstSharedStruct from templated struct instance. This will create a new instance of the shared struct memory. */
 	template<typename T = BaseStructT>
 		requires (std::is_base_of_v<BaseStructT, std::decay_t<T>>)
 	static TConstSharedStruct Make(const T& Struct)
-	{
-		FConstSharedStruct SharedStruct;
-		SharedStruct.InitializeAs<BaseStructT>(Struct);
-		return SharedStruct;
-	}
-
-	/** Creates a new TSharedStruct from a struct type and optional data. This will create a new instance of the shared struct memory. */
-	template<typename T = BaseStructT>
-		requires (std::is_base_of_v<BaseStructT, std::decay_t<T>>)
-	static TConstSharedStruct Make(const uint8* InStructMemory = nullptr)
 	{
 		TConstSharedStruct SharedStruct;
-		SharedStruct.template Initialize<T>(InStructMemory);
+		SharedStruct.Initialize<BaseStructT>(Struct);
 		return SharedStruct;
 	}
 
-	/** Creates a new TSharedStruct from the templated type and forward all arguments to constructor. This will create a new instance of the shared struct memory. */
+	/** Creates a new TConstSharedStruct from the templated type and forward all arguments to constructor. This will create a new instance of the shared struct memory. */
 	template<typename T = BaseStructT, typename... TArgs>
 		requires (std::is_base_of_v<BaseStructT, std::decay_t<T>>)
 	static TConstSharedStruct Make(TArgs&&... InArgs)
@@ -879,15 +879,7 @@ struct [[nodiscard]] TConstSharedStruct
 	/** Returns const reference to the struct, this getter assumes that all data is valid. */
 	template<typename T = BaseStructT>
 		requires (std::is_base_of_v<BaseStructT, std::decay_t<T>>)
-	constexpr const T& Get() const
-	{
-		return ConstSharedStruct.Get<T>();
-	}
-
-	/** Returns const reference to the struct, this getter assumes that all data is valid. */
-	template<typename T = BaseStructT>
-		requires (std::is_const_v<T> && std::is_base_of_v<BaseStructT, std::decay_t<T>>)
-	constexpr T& Get() const
+	const T& Get() const
 	{
 		return ConstSharedStruct.Get<T>();
 	}
@@ -895,15 +887,7 @@ struct [[nodiscard]] TConstSharedStruct
 	/** Returns const pointer to the struct. */
 	template<typename T = BaseStructT>
 		requires (std::is_base_of_v<BaseStructT, std::decay_t<T>>)
-	constexpr const T* GetPtr() const
-	{
-		return ConstSharedStruct.GetPtr<T>();
-	}
-
-	/** Returns const pointer to the struct, or nullptr if cast is not valid. */
-	template<typename T = BaseStructT>
-		requires (std::is_const_v<T> && std::is_base_of_v<BaseStructT, std::decay_t<T>>)
-	constexpr T* GetPtr() const
+	const T* GetPtr() const
 	{
 		return ConstSharedStruct.GetPtr<T>();
 	}
@@ -917,16 +901,30 @@ struct [[nodiscard]] TConstSharedStruct
 	/** Comparison operators. Note: it does not compare the internal structure itself */
 	template<typename T = BaseStructT>
 		requires (std::is_base_of_v<BaseStructT, std::decay_t<T>>)
-	bool operator==(const T& Other) const
+	bool operator==(const TSharedStruct<T>& Other) const
 	{
-		return ConstSharedStruct.operator==(Other);
+		return ConstSharedStruct.operator==(Other.SharedStruct);
 	}
 
 	template<typename T = BaseStructT>
 		requires (std::is_base_of_v<BaseStructT, std::decay_t<T>>)
-	bool operator!=(const T& Other) const
+	bool operator!=(const TSharedStruct<T>& Other) const
 	{
-		return ConstSharedStruct.operator!=(Other);
+		return ConstSharedStruct.operator!=(Other.SharedStruct);
+	}
+
+	template<typename T = BaseStructT>
+		requires (std::is_base_of_v<BaseStructT, std::decay_t<T>>)
+	bool operator==(const TConstSharedStruct<T>& Other) const
+	{
+		return ConstSharedStruct.operator==(Other.ConstSharedStruct);
+	}
+
+	template<typename T = BaseStructT>
+		requires (std::is_base_of_v<BaseStructT, std::decay_t<T>>)
+	bool operator!=(const TConstSharedStruct<T>& Other) const
+	{
+		return ConstSharedStruct.operator!=(Other.ConstSharedStruct);
 	}
 
 private:
@@ -940,6 +938,7 @@ private:
 	 */
 	FConstSharedStruct ConstSharedStruct;
 
+	template <typename U> friend struct TSharedStruct;
 	template <typename U> friend struct TConstSharedStruct;
 };
 
