[StructUtils] Fix serialization with default when the struct is zero-initialized. Because it uses nullptr as the default, if the struct changes in the future, the defaults will not be taken into account. Add unittest for FInstancedStruct::Serialize and FInstancedStructContainer::Serialize.

#Synced-CL 52583850

diff --git a/Engine/Source/Developer/StructUtilsTestSuite/Private/InstancedStructContainerTest.cpp b/Engine/Source/Developer/StructUtilsTestSuite/Private/InstancedStructContainerTest.cpp
--- a/Engine/Source/Developer/StructUtilsTestSuite/Private/InstancedStructContainerTest.cpp
+++ b/Engine/Source/Developer/StructUtilsTestSuite/Private/InstancedStructContainerTest.cpp
@@ -181,6 +181,105 @@ struct FTest_RangedForInstancedStructContainer : FAITestBase
 };
 IMPLEMENT_AI_INSTANT_TEST(FTest_RangedForInstancedStructContainer, "System.StructUtils.InstancedStructContainer.RangedFor");
 
+// Test that FInstancedStructContainer serialization correctly preserves values both with explicit
+// Defaults and with nullptr DefaultsStruct / nullptr Defaults (C++ struct defaults fallback).
+// The same four-item layout is exercised for each path:
+//   Item 0: FTestStructSimpleNonZeroDefault, Float = 100.0f (== default 100.0f) — not written.
+//   Item 1: FTestStructSimpleNonZeroDefault, Float =   0.0f (!= default 100.0f) — written.
+//   Item 2: FTestStructSimple,               Float =   0.0f (== default   0.0f) — not written.
+//   Item 3: FTestStructSimple,               Float =  50.0f (!= default   0.0f) — written.
+struct FTest_SerializeInstancedStructContainerWithDefaults : FAITestBase
+{
+	virtual bool InstantTest() override
+	{
+		// Helper lambda that builds a container with the four test items.
+		auto MakeContainer = [](FInstancedStructContainer& Container)
+		{
+			TArray<FInstancedStruct> TestInstanced;
+			TestInstanced.AddDefaulted_GetRef().InitializeAs<FTestStructSimpleNonZeroDefault>();                                          // Item 0: Float=100.0f
+			FTestStructSimpleNonZeroDefault& Item1 = TestInstanced.AddDefaulted_GetRef().InitializeAs<FTestStructSimpleNonZeroDefault>(); // Item 1
+			Item1.Float = 0.0f;
+			TestInstanced.AddDefaulted_GetRef().InitializeAs<FTestStructSimple>();                                                        // Item 2: Float=0.0f
+			FTestStructSimple& Item3 = TestInstanced.AddDefaulted_GetRef().InitializeAs<FTestStructSimple>();                             // Item 3
+			Item3.Float = 50.0f;
+			Container.Append(TestInstanced);
+		};
+
+		// Helper lambda that scrambles all four items so a subsequent load must actively restore them.
+		auto ScrambleItems = [](FInstancedStructContainer& Container)
+		{
+			Container[0].Get<FTestStructSimpleNonZeroDefault>().Float = 0.0f;   // was 100.0f
+			Container[1].Get<FTestStructSimpleNonZeroDefault>().Float = 100.0f; // was 0.0f
+			Container[2].Get<FTestStructSimple>().Float = 50.0f;                // was 0.0f
+			Container[3].Get<FTestStructSimple>().Float = 0.0f;                 // was 50.0f
+		};
+
+		// --- Pass 1: explicit DefaultContainer ---
+		{
+			// Build the defaults container — one entry per item, representing archetype/CDO values.
+			FInstancedStructContainer DefaultContainer;
+			{
+				TArray<FInstancedStruct> DefaultInstanced;
+				DefaultInstanced.AddDefaulted_GetRef().InitializeAs<FTestStructSimpleNonZeroDefault>(); // Float=100.0f
+				DefaultInstanced.AddDefaulted_GetRef().InitializeAs<FTestStructSimpleNonZeroDefault>(); // Float=100.0f
+				DefaultInstanced.AddDefaulted_GetRef().InitializeAs<FTestStructSimple>();               // Float=0.0f
+				DefaultInstanced.AddDefaulted_GetRef().InitializeAs<FTestStructSimple>();               // Float=0.0f
+				DefaultContainer.Append(DefaultInstanced);
+			}
+
+			FInstancedStructContainer Container;
+			MakeContainer(Container);
+			AITEST_EQUAL(TEXT("Explicit: Should have 4 items before save"), Container.Num(), 4);
+
+			TArray<uint8> Memory;
+			FMemoryWriter Writer(Memory);
+			FObjectAndNameAsStringProxyArchive WriterProxy(Writer, /*bInLoadIfFindFails*/false);
+			AITEST_TRUE(TEXT("Explicit: saving should succeed"), Container.Serialize(WriterProxy, FInstancedStructContainer::StaticStruct(), &DefaultContainer));
+
+			ScrambleItems(Container);
+
+			FMemoryReader Reader(Memory);
+			FObjectAndNameAsStringProxyArchive ReaderProxy(Reader, /*bInLoadIfFindFails*/true);
+			AITEST_TRUE(TEXT("Explicit: loading should succeed"), Container.Serialize(ReaderProxy, FInstancedStructContainer::StaticStruct(), &DefaultContainer));
+
+			AITEST_EQUAL(TEXT("Explicit: Should have 4 items after load"), Container.Num(), 4);
+			AITEST_EQUAL(TEXT("Explicit: Item 0 Float should be restored to default 100.0f"), 100.0f, Container[0].Get<FTestStructSimpleNonZeroDefault>().Float);
+			AITEST_EQUAL(TEXT("Explicit: Item 1 Float should be restored to saved zero 0.0f"),   0.0f, Container[1].Get<FTestStructSimpleNonZeroDefault>().Float);
+			AITEST_EQUAL(TEXT("Explicit: Item 2 Float should be restored to default 0.0f"),      0.0f, Container[2].Get<FTestStructSimple>().Float);
+			AITEST_EQUAL(TEXT("Explicit: Item 3 Float should be restored to saved value 50.0f"), 50.0f, Container[3].Get<FTestStructSimple>().Float);
+		}
+
+		// --- Pass 2: nullptr DefaultsStruct and nullptr Defaults ---
+		// The serializer falls back to a temporary C++ default instance per item, so results
+		// are identical to Pass 1 because the explicit defaults above equal the C++ defaults.
+		{
+			FInstancedStructContainer Container;
+			MakeContainer(Container);
+			AITEST_EQUAL(TEXT("Nullptr: Should have 4 items before save"), Container.Num(), 4);
+
+			TArray<uint8> Memory;
+			FMemoryWriter Writer(Memory);
+			FObjectAndNameAsStringProxyArchive WriterProxy(Writer, /*bInLoadIfFindFails*/false);
+			AITEST_TRUE(TEXT("Nullptr: saving should succeed"), Container.Serialize(WriterProxy, nullptr, nullptr));
+
+			ScrambleItems(Container);
+
+			FMemoryReader Reader(Memory);
+			FObjectAndNameAsStringProxyArchive ReaderProxy(Reader, /*bInLoadIfFindFails*/true);
+			AITEST_TRUE(TEXT("Nullptr: loading should succeed"), Container.Serialize(ReaderProxy, nullptr, nullptr));
+
+			AITEST_EQUAL(TEXT("Nullptr: Should have 4 items after load"), Container.Num(), 4);
+			AITEST_EQUAL(TEXT("Nullptr: Item 0 Float should be restored to C++ default 100.0f"), 100.0f, Container[0].Get<FTestStructSimpleNonZeroDefault>().Float);
+			AITEST_EQUAL(TEXT("Nullptr: Item 1 Float should be restored to saved zero 0.0f"),      0.0f, Container[1].Get<FTestStructSimpleNonZeroDefault>().Float);
+			AITEST_EQUAL(TEXT("Nullptr: Item 2 Float should be restored to C++ default 0.0f"),     0.0f, Container[2].Get<FTestStructSimple>().Float);
+			AITEST_EQUAL(TEXT("Nullptr: Item 3 Float should be restored to saved value 50.0f"),   50.0f, Container[3].Get<FTestStructSimple>().Float);
+		}
+
+		return true;
+	}
+};
+IMPLEMENT_AI_INSTANT_TEST(FTest_SerializeInstancedStructContainerWithDefaults, "System.StructUtils.InstancedStructContainer.SerializeWithDefaults");
+
 } // FInstancedStructContainerTest
 
 #undef LOCTEXT_NAMESPACE
diff --git a/Engine/Source/Developer/StructUtilsTestSuite/Private/InstancedStructTest.cpp b/Engine/Source/Developer/StructUtilsTestSuite/Private/InstancedStructTest.cpp
--- a/Engine/Source/Developer/StructUtilsTestSuite/Private/InstancedStructTest.cpp
+++ b/Engine/Source/Developer/StructUtilsTestSuite/Private/InstancedStructTest.cpp
@@ -3,6 +3,9 @@
 #include "AITestsCommon.h"
 #include "StructUtils/SharedStruct.h"
 #include "StructUtilsTestTypes.h"
+#include "Serialization/ObjectAndNameAsStringProxyArchive.h"
+#include "Serialization/MemoryWriter.h"
+#include "Serialization/MemoryReader.h"
 
 #define LOCTEXT_NAMESPACE "StructUtilsTests"
 
@@ -156,6 +159,188 @@ namespace FInstancedStructTest
 	};
 	IMPLEMENT_AI_INSTANT_TEST(FTest_InstancedStructCustomScriptStruct, "System.StructUtils.InstancedStruct.CustomScriptStruct");
 
+	// Test that serialization with explicit Defaults correctly preserves values across four cases,
+	// and that passing nullptr/nullptr falls back to C++ defaults with the same correctness.
+	// The Defaults FInstancedStruct drives both what gets skipped during save (delta) and what
+	// gets used as the base values during load (InitializeAs is called with the default memory
+	// before the stored delta is applied).
+	//   Cases 1-4: explicit Defaults.
+	//   Cases 5-8: nullptr DefaultsStruct / nullptr Defaults (falls back to C++ struct defaults).
+	//   Case 1/5: FTestStructSimpleNonZeroDefault, Float == default (100.0f) — not written.
+	//   Case 2/6: FTestStructSimpleNonZeroDefault, Float != default (0.0f vs 100.0f) — written.
+	//   Case 3/7: FTestStructSimple, Float == default (0.0f) — not written.
+	//   Case 4/8: FTestStructSimple, Float != default (50.0f vs 0.0f) — written.
+	struct FTest_InstancedStructSerializeWithDefaults : FAITestBase
+	{
+		virtual bool InstantTest() override
+		{
+			// Case 1: FTestStructSimpleNonZeroDefault, Float at the provided default (100.0f).
+			// Not written; load must restore it from the default.
+			{
+				const FInstancedStruct DefaultStruct = FInstancedStruct::Make<FTestStructSimpleNonZeroDefault>(); // Float=100.0f
+
+				FInstancedStruct InstancedStruct = FInstancedStruct::Make<FTestStructSimpleNonZeroDefault>(); // Float=100.0f
+
+				TArray<uint8> Memory;
+				FMemoryWriter Writer(Memory);
+				FObjectAndNameAsStringProxyArchive WriterProxy(Writer, /*bInLoadIfFindFails*/false);
+				AITEST_TRUE(TEXT("Case1: saving should succeed"), InstancedStruct.Serialize(WriterProxy, FInstancedStruct::StaticStruct(), &DefaultStruct));
+
+				InstancedStruct.GetMutable<FTestStructSimpleNonZeroDefault>().Float = 0.0f;
+
+				FMemoryReader Reader(Memory);
+				FObjectAndNameAsStringProxyArchive ReaderProxy(Reader, /*bInLoadIfFindFails*/true);
+				AITEST_TRUE(TEXT("Case1: loading should succeed"), InstancedStruct.Serialize(ReaderProxy, FInstancedStruct::StaticStruct(), &DefaultStruct));
+
+				AITEST_EQUAL(TEXT("Case1: Float should be restored to default 100.0f"), 100.0f, InstancedStruct.Get<FTestStructSimpleNonZeroDefault>().Float);
+			}
+
+			// Case 2: FTestStructSimpleNonZeroDefault, Float differs from the provided default (0.0f vs 100.0f).
+			// Written; load must restore the saved zero value.
+			{
+				const FInstancedStruct DefaultStruct = FInstancedStruct::Make<FTestStructSimpleNonZeroDefault>(); // Float=100.0f
+
+				FInstancedStruct InstancedStruct = FInstancedStruct::Make<FTestStructSimpleNonZeroDefault>();
+				InstancedStruct.GetMutable<FTestStructSimpleNonZeroDefault>().Float = 0.0f;
+
+				TArray<uint8> Memory;
+				FMemoryWriter Writer(Memory);
+				FObjectAndNameAsStringProxyArchive WriterProxy(Writer, /*bInLoadIfFindFails*/false);
+				AITEST_TRUE(TEXT("Case2: saving should succeed"), InstancedStruct.Serialize(WriterProxy, FInstancedStruct::StaticStruct(), &DefaultStruct));
+
+				InstancedStruct.GetMutable<FTestStructSimpleNonZeroDefault>().Float = 100.0f;
+
+				FMemoryReader Reader(Memory);
+				FObjectAndNameAsStringProxyArchive ReaderProxy(Reader, /*bInLoadIfFindFails*/true);
+				AITEST_TRUE(TEXT("Case2: loading should succeed"), InstancedStruct.Serialize(ReaderProxy, FInstancedStruct::StaticStruct(), &DefaultStruct));
+
+				AITEST_EQUAL(TEXT("Case2: Float should be restored to saved zero 0.0f"), 0.0f, InstancedStruct.Get<FTestStructSimpleNonZeroDefault>().Float);
+			}
+
+			// Case 3: FTestStructSimple, Float at the provided default (0.0f).
+			// Not written; load must restore it from the default.
+			{
+				const FInstancedStruct DefaultStruct = FInstancedStruct::Make<FTestStructSimple>(); // Float=0.0f
+
+				FInstancedStruct InstancedStruct = FInstancedStruct::Make<FTestStructSimple>(); // Float=0.0f
+
+				TArray<uint8> Memory;
+				FMemoryWriter Writer(Memory);
+				FObjectAndNameAsStringProxyArchive WriterProxy(Writer, /*bInLoadIfFindFails*/false);
+				AITEST_TRUE(TEXT("Case3: saving should succeed"), InstancedStruct.Serialize(WriterProxy, FInstancedStruct::StaticStruct(), &DefaultStruct));
+
+				InstancedStruct.GetMutable<FTestStructSimple>().Float = 50.0f;
+
+				FMemoryReader Reader(Memory);
+				FObjectAndNameAsStringProxyArchive ReaderProxy(Reader, /*bInLoadIfFindFails*/true);
+				AITEST_TRUE(TEXT("Case3: loading should succeed"), InstancedStruct.Serialize(ReaderProxy, FInstancedStruct::StaticStruct(), &DefaultStruct));
+
+				AITEST_EQUAL(TEXT("Case3: Float should be restored to default 0.0f"), 0.0f, InstancedStruct.Get<FTestStructSimple>().Float);
+			}
+
+			// Case 4: FTestStructSimple, Float differs from the provided default (50.0f vs 0.0f).
+			// Written; load must restore the saved non-zero value.
+			{
+				const FInstancedStruct DefaultStruct = FInstancedStruct::Make<FTestStructSimple>(); // Float=0.0f
+
+				FInstancedStruct InstancedStruct = FInstancedStruct::Make<FTestStructSimple>(50.0f);
+
+				TArray<uint8> Memory;
+				FMemoryWriter Writer(Memory);
+				FObjectAndNameAsStringProxyArchive WriterProxy(Writer, /*bInLoadIfFindFails*/false);
+				AITEST_TRUE(TEXT("Case4: saving should succeed"), InstancedStruct.Serialize(WriterProxy, FInstancedStruct::StaticStruct(), &DefaultStruct));
+
+				InstancedStruct.GetMutable<FTestStructSimple>().Float = 0.0f;
+
+				FMemoryReader Reader(Memory);
+				FObjectAndNameAsStringProxyArchive ReaderProxy(Reader, /*bInLoadIfFindFails*/true);
+				AITEST_TRUE(TEXT("Case4: loading should succeed"), InstancedStruct.Serialize(ReaderProxy, FInstancedStruct::StaticStruct(), &DefaultStruct));
+
+				AITEST_EQUAL(TEXT("Case4: Float should be restored to saved value 50.0f"), 50.0f, InstancedStruct.Get<FTestStructSimple>().Float);
+			}
+
+			// Cases 5-8: same four scenarios with nullptr DefaultsStruct and nullptr Defaults.
+			// The serializer falls back to a temporary C++ default instance, so the behaviour is
+			// identical to the explicit-defaults cases above (whose defaults match the C++ defaults).
+
+			// Case 5: FTestStructSimpleNonZeroDefault, Float at C++ default (100.0f). Not written.
+			{
+				FInstancedStruct InstancedStruct = FInstancedStruct::Make<FTestStructSimpleNonZeroDefault>(); // Float=100.0f
+
+				TArray<uint8> Memory;
+				FMemoryWriter Writer(Memory);
+				FObjectAndNameAsStringProxyArchive WriterProxy(Writer, /*bInLoadIfFindFails*/false);
+				AITEST_TRUE(TEXT("Case5: saving should succeed"), InstancedStruct.Serialize(WriterProxy, nullptr, nullptr));
+
+				InstancedStruct.GetMutable<FTestStructSimpleNonZeroDefault>().Float = 0.0f;
+
+				FMemoryReader Reader(Memory);
+				FObjectAndNameAsStringProxyArchive ReaderProxy(Reader, /*bInLoadIfFindFails*/true);
+				AITEST_TRUE(TEXT("Case5: loading should succeed"), InstancedStruct.Serialize(ReaderProxy, nullptr, nullptr));
+
+				AITEST_EQUAL(TEXT("Case5: Float should be restored to C++ default 100.0f"), 100.0f, InstancedStruct.Get<FTestStructSimpleNonZeroDefault>().Float);
+			}
+
+			// Case 6: FTestStructSimpleNonZeroDefault, Float differs from C++ default (0.0f vs 100.0f). Written.
+			{
+				FInstancedStruct InstancedStruct = FInstancedStruct::Make<FTestStructSimpleNonZeroDefault>();
+				InstancedStruct.GetMutable<FTestStructSimpleNonZeroDefault>().Float = 0.0f;
+
+				TArray<uint8> Memory;
+				FMemoryWriter Writer(Memory);
+				FObjectAndNameAsStringProxyArchive WriterProxy(Writer, /*bInLoadIfFindFails*/false);
+				AITEST_TRUE(TEXT("Case6: saving should succeed"), InstancedStruct.Serialize(WriterProxy, nullptr, nullptr));
+
+				InstancedStruct.GetMutable<FTestStructSimpleNonZeroDefault>().Float = 100.0f;
+
+				FMemoryReader Reader(Memory);
+				FObjectAndNameAsStringProxyArchive ReaderProxy(Reader, /*bInLoadIfFindFails*/true);
+				AITEST_TRUE(TEXT("Case6: loading should succeed"), InstancedStruct.Serialize(ReaderProxy, nullptr, nullptr));
+
+				AITEST_EQUAL(TEXT("Case6: Float should be restored to saved zero 0.0f"), 0.0f, InstancedStruct.Get<FTestStructSimpleNonZeroDefault>().Float);
+			}
+
+			// Case 7: FTestStructSimple, Float at C++ default (0.0f). Not written.
+			{
+				FInstancedStruct InstancedStruct = FInstancedStruct::Make<FTestStructSimple>(); // Float=0.0f
+
+				TArray<uint8> Memory;
+				FMemoryWriter Writer(Memory);
+				FObjectAndNameAsStringProxyArchive WriterProxy(Writer, /*bInLoadIfFindFails*/false);
+				AITEST_TRUE(TEXT("Case7: saving should succeed"), InstancedStruct.Serialize(WriterProxy, nullptr, nullptr));
+
+				InstancedStruct.GetMutable<FTestStructSimple>().Float = 50.0f;
+
+				FMemoryReader Reader(Memory);
+				FObjectAndNameAsStringProxyArchive ReaderProxy(Reader, /*bInLoadIfFindFails*/true);
+				AITEST_TRUE(TEXT("Case7: loading should succeed"), InstancedStruct.Serialize(ReaderProxy, nullptr, nullptr));
+
+				AITEST_EQUAL(TEXT("Case7: Float should be restored to C++ default 0.0f"), 0.0f, InstancedStruct.Get<FTestStructSimple>().Float);
+			}
+
+			// Case 8: FTestStructSimple, Float differs from C++ default (50.0f vs 0.0f). Written.
+			{
+				FInstancedStruct InstancedStruct = FInstancedStruct::Make<FTestStructSimple>(50.0f);
+
+				TArray<uint8> Memory;
+				FMemoryWriter Writer(Memory);
+				FObjectAndNameAsStringProxyArchive WriterProxy(Writer, /*bInLoadIfFindFails*/false);
+				AITEST_TRUE(TEXT("Case8: saving should succeed"), InstancedStruct.Serialize(WriterProxy, nullptr, nullptr));
+
+				InstancedStruct.GetMutable<FTestStructSimple>().Float = 0.0f;
+
+				FMemoryReader Reader(Memory);
+				FObjectAndNameAsStringProxyArchive ReaderProxy(Reader, /*bInLoadIfFindFails*/true);
+				AITEST_TRUE(TEXT("Case8: loading should succeed"), InstancedStruct.Serialize(ReaderProxy, nullptr, nullptr));
+
+				AITEST_EQUAL(TEXT("Case8: Float should be restored to saved value 50.0f"), 50.0f, InstancedStruct.Get<FTestStructSimple>().Float);
+			}
+
+			return true;
+		}
+	};
+	IMPLEMENT_AI_INSTANT_TEST(FTest_InstancedStructSerializeWithDefaults, "System.StructUtils.InstancedStruct.SerializeWithDefaults");
+
 }
 
 
diff --git a/Engine/Source/Runtime/CoreUObject/Private/StructUtils/InstancedStruct.cpp b/Engine/Source/Runtime/CoreUObject/Private/StructUtils/InstancedStruct.cpp
--- a/Engine/Source/Runtime/CoreUObject/Private/StructUtils/InstancedStruct.cpp
+++ b/Engine/Source/Runtime/CoreUObject/Private/StructUtils/InstancedStruct.cpp
@@ -143,7 +143,7 @@ bool FInstancedStruct::Serialize(FArchive& Ar, UStruct* DefaultsStruct, const vo
 					// Use the parent Defaults
 					ScriptStruct->SerializeItem(Ar, Self->GetMutableMemory(), DefaultInstancedStruct->GetMemory());
 				}
-				else if (UScriptStruct::ICppStructOps* TheCppStructOps = ScriptStruct->GetCppStructOps(); TheCppStructOps && !TheCppStructOps->HasZeroConstructor())
+				else
 				{
 					// Create a temporary struct on the stack. The struct has default values.
 					const int32 MinAlignment = ScriptStruct->GetMinAlignment();
@@ -155,11 +155,6 @@ bool FInstancedStruct::Serialize(FArchive& Ar, UStruct* DefaultsStruct, const vo
 					ScriptStruct->SerializeItem(Ar, Self->GetMutableMemory(), TempMemory);
 					ScriptStruct->DestroyStruct(TempMemory);
 				}
-				else
-				{
-					// Serialize without defaults
-					ScriptStruct->SerializeItem(Ar, Self->GetMutableMemory(), nullptr);
-				}
 			}
 			else
 			{
diff --git a/Engine/Source/Runtime/CoreUObject/Private/StructUtils/InstancedStructContainer.cpp b/Engine/Source/Runtime/CoreUObject/Private/StructUtils/InstancedStructContainer.cpp
--- a/Engine/Source/Runtime/CoreUObject/Private/StructUtils/InstancedStructContainer.cpp
+++ b/Engine/Source/Runtime/CoreUObject/Private/StructUtils/InstancedStructContainer.cpp
@@ -519,7 +519,7 @@ bool FInstancedStructContainer::Serialize(FArchive& Ar, UStruct* DefaultsStruct,
 					const void* DefaultMemory = DefaultInstancedStructContainer->Memory + DefaultInstancedStructContainer->GetItem(Index).Offset;
 					ConstCast(Item.ScriptStruct)->SerializeItem(Ar, Self->Memory + Item.Offset, DefaultMemory);
 				}
-				else if (UScriptStruct::ICppStructOps* TheCppStructOps = Item.ScriptStruct->GetCppStructOps(); TheCppStructOps && !TheCppStructOps->HasZeroConstructor())
+				else
 				{
 					// Create a temporary struct on the stack. The struct has default values.
 					const int32 MinAlignment = Item.ScriptStruct->GetMinAlignment();
@@ -531,11 +531,6 @@ bool FInstancedStructContainer::Serialize(FArchive& Ar, UStruct* DefaultsStruct,
 					ConstCast(Item.ScriptStruct)->SerializeItem(Ar, Self->Memory + Item.Offset, TempMemory);
 					ConstCast(Item.ScriptStruct)->DestroyStruct(TempMemory);
 				}
-				else
-				{
-					// Serialize without defaults
-					ConstCast(Item.ScriptStruct)->SerializeItem(Ar, Self->Memory + Item.Offset, nullptr);
-				}
 			}
 			else
 			{
