[StructUtils] Make sure the property bag instance is constructed and destructed to increment/decrement the ref counting

* Add unit tests too

#rnx
#Synced-CL 52681881

diff --git a/Engine/Source/Developer/StructUtilsTestSuite/Private/PropertyBagTest.cpp b/Engine/Source/Developer/StructUtilsTestSuite/Private/PropertyBagTest.cpp
--- a/Engine/Source/Developer/StructUtilsTestSuite/Private/PropertyBagTest.cpp
+++ b/Engine/Source/Developer/StructUtilsTestSuite/Private/PropertyBagTest.cpp
@@ -1640,6 +1640,51 @@ struct FTest_SameBag : FAITestBase
 };
 IMPLEMENT_AI_INSTANT_TEST(FTest_SameBag, "System.StructUtils.PropertyBag.SameBag");
 
+struct FTest_StructOfStructRefCounting : FAITestBase
+{
+	virtual bool InstantTest() override
+	{
+		static const FName InnerName(TEXT("Inner"));
+		static const FName StringName(TEXT("String"));
+
+		// Create a first bag with an instance, and then a second bag that has the first bag as member.
+		// Use a string type, because with previous behavior, the property bag would be marked "ZeroConstructor" (since String can be zero constructed)
+		// but will still have to call a destructor (because strings are not trivially destructible). 
+		// It is expected that the refcount of the first bag would be 2, and the ref count of the second bag to be 1.
+		// Since RefCounting of the property bag is protected, we can't check explicitly the refcount. But
+		// if refcounting is not correct, releasing the instances would throw an error and will make the test fail.
+		// At the end of the block, instances will be released and will go through decrementing the refcount.
+		// We'll also check that the property bags are not zero constructed and have a destructor (in order for the ref count to be incremented/decremented)
+
+		{
+			FInstancedPropertyBag InnerBag;
+			InnerBag.AddProperties({
+				{ StringName, EPropertyBagPropertyType::String },
+			});
+
+			const UPropertyBag* InnerPropertyBag = InnerBag.GetPropertyBagStruct();
+			AITEST_NOT_NULL(TEXT("Inner property bag is valid"), InnerPropertyBag);
+			AITEST_TRUE(TEXT("Inner property bag is not zero constructible"), (InnerPropertyBag->StructFlags & STRUCT_ZeroConstructor) == 0);
+			AITEST_TRUE(TEXT("Inner property bag has a destructor"), (InnerPropertyBag->StructFlags & STRUCT_NoDestructor) == 0);
+			AITEST_TRUE(TEXT("Inner property bag is not plain old data"), (InnerPropertyBag->StructFlags & STRUCT_IsPlainOldData) == 0);
+
+			FInstancedPropertyBag OuterBag;
+			OuterBag.AddProperties({
+				{ InnerName, EPropertyBagContainerType::None, EPropertyBagPropertyType::Struct, InnerPropertyBag }
+			});
+		
+			const UPropertyBag* OuterPropertyBag = OuterBag.GetPropertyBagStruct();
+			AITEST_NOT_NULL(TEXT("Outer property bag is valid"), OuterPropertyBag);
+			AITEST_TRUE(TEXT("Outer property bag is not zero constructible"), (OuterPropertyBag->StructFlags & STRUCT_ZeroConstructor) == 0);
+			AITEST_TRUE(TEXT("Outer property bag has a destructor"), (OuterPropertyBag->StructFlags & STRUCT_NoDestructor) == 0);
+			AITEST_TRUE(TEXT("Outer property bag is not plain old data"), (OuterPropertyBag->StructFlags & STRUCT_IsPlainOldData) == 0);
+		}
+
+		return true;
+	}
+};
+IMPLEMENT_AI_INSTANT_TEST(FTest_StructOfStructRefCounting, "System.StructUtils.PropertyBag.StructOfStructRefCounting");
+
 } // FPropertyBagTest
 
 #undef LOCTEXT_NAMESPACE
diff --git a/Engine/Source/Runtime/CoreUObject/Private/StructUtils/PropertyBag.cpp b/Engine/Source/Runtime/CoreUObject/Private/StructUtils/PropertyBag.cpp
--- a/Engine/Source/Runtime/CoreUObject/Private/StructUtils/PropertyBag.cpp
+++ b/Engine/Source/Runtime/CoreUObject/Private/StructUtils/PropertyBag.cpp
@@ -4190,6 +4190,17 @@ const UPropertyBag* UPropertyBag::GetOrCreateFromDescs(const TConstArrayView<FPr
 	return NewBag;
 }
 
+void UPropertyBag::Link(FArchive& Ar, bool bRelinkExistingProperties)
+{
+	Super::Link(Ar, bRelinkExistingProperties);
+	
+	// Every time an instance of this property bag is initialized we need to increment the ref count by calling `InitializeStruct`.
+	// And when it is destroyed, it decrements the ref count by calling `DestroyStruct`.
+	// Therefore, the property bag can never be zero constructed and must have a destructor.
+	// Same logic applies to "plain old data", we can never be memcpy.
+	StructFlags = static_cast<EStructFlags>(StructFlags & ~(STRUCT_ZeroConstructor | STRUCT_NoDestructor | STRUCT_IsPlainOldData));
+}
+
 #if WITH_EDITOR
 bool UPropertyBag::ContainsUserDefinedStruct(const UUserDefinedStruct* UserDefinedStruct) const
 {
diff --git a/Engine/Source/Runtime/CoreUObject/Public/StructUtils/PropertyBag.h b/Engine/Source/Runtime/CoreUObject/Public/StructUtils/PropertyBag.h
--- a/Engine/Source/Runtime/CoreUObject/Public/StructUtils/PropertyBag.h
+++ b/Engine/Source/Runtime/CoreUObject/Public/StructUtils/PropertyBag.h
@@ -1453,6 +1453,7 @@ protected:
 	UE_API virtual void InitializeStruct(void* Dest, int32 ArrayDim = 1) const override;
 	UE_API virtual void DestroyStruct(void* Dest, int32 ArrayDim = 1) const override;
 	UE_API virtual void FinishDestroy() override;
+	UE_API virtual void Link(FArchive& Ar, bool bRelinkExistingProperties) override;
 
 	virtual bool ArePropertyGuidsAvailable() const override
 	{
